Fixed sky diffuse being invisible in project with movable skylight and r.SupportStationarySkylight is false.
Behavior was changed in 4691984.

SupportStationarySkylight should not be used to conditionally create a sky or not. Sky is always used for default environment diffuse (just can be rendered in non base pass but environment pass depending on whether lighting/sky is static/stationary of movable).
A SphereReflectionCapture is not enough to get diffuse light (it can only be used for specular). Which is what we see in this JIRA bug.

Fix is to always create the sky and just make it visible or not (the new behavior added in 4691984 and that is correct)

diff --git a/Engine/Source/Editor/AdvancedPreviewScene/Private/AdvancedPreviewScene.cpp b/Engine/Source/Editor/AdvancedPreviewScene/Private/AdvancedPreviewScene.cpp
--- a/Engine/Source/Editor/AdvancedPreviewScene/Private/AdvancedPreviewScene.cpp
+++ b/Engine/Source/Editor/AdvancedPreviewScene/Private/AdvancedPreviewScene.cpp
@@ -38,32 +38,11 @@ FAdvancedPreviewScene::FAdvancedPreviewScene(ConstructionValues CVS, float InFlo
 	FPreviewSceneProfile& Profile = DefaultSettings->Profiles[CurrentProfileIndex];
 	Profile.LoadEnvironmentMap();
 
-	SphereReflectionComponent = nullptr;
-
 	const FTransform Transform(FRotator(0, 0, 0), FVector(0, 0, 0), FVector(1));
 
-	static const auto CVarSupportStationarySkylight = IConsoleManager::Get().FindTConsoleVariableDataInt(TEXT("r.SupportStationarySkylight"));
-	bUseSkylight = CVarSupportStationarySkylight->GetValueOnAnyThread() != 0;
-
-	if (bUseSkylight)
-	{
-		// Set up sky light using the set cube map texture, reusing the sky light from PreviewScene class
-		SetSkyCubemap(Profile.EnvironmentCubeMap.Get());
-		SetSkyBrightness(Profile.SkyLightIntensity);
-	}
-	else
-	{
-		// Hide the inherited skylight 
-		SkyLight->SetVisibility(false);
-
-		// Setup sphere reflection
-		SphereReflectionComponent = NewObject<USphereReflectionCaptureComponent>();
-		SphereReflectionComponent->Cubemap = Profile.EnvironmentCubeMap.Get();
-		SphereReflectionComponent->ReflectionSourceType = EReflectionSourceType::SpecifiedCubemap;
-		SphereReflectionComponent->Brightness = Profile.SkyLightIntensity;
-		AddComponent(SphereReflectionComponent, Transform);
-		SphereReflectionComponent->UpdateReflectionCaptureContents(PreviewWorld);
-	}
+	// Always set up sky light using the set cube map texture, reusing the sky light from PreviewScene class
+	SetSkyCubemap(Profile.EnvironmentCubeMap.Get());
+	SetSkyBrightness(Profile.SkyLightIntensity);
 	
 	// Large scale to prevent sphere from clipping
 	const FTransform SphereTransform(FRotator(0, 0, 0), FVector(0, 0, 0), FVector(2000));
@@ -135,17 +114,10 @@ void FAdvancedPreviewScene::UpdateScene(FPreviewSceneProfile& Profile, bool bUpd
 	if (bUpdateSkyLight)
 	{
 		// Threshold to ensure we only update the intensity if it is going to make a difference
-		if (!FMath::IsNearlyEqual(bUseSkylight ? SkyLight->Intensity : SphereReflectionComponent->Brightness, Profile.SkyLightIntensity, 0.05f))
+		if (!FMath::IsNearlyEqual(SkyLight->Intensity, Profile.SkyLightIntensity, 0.05f))
 		{
 			static const FName IntensityName("Intensity");
-			if (bUseSkylight)
-			{
-				SetSkyBrightness(Profile.SkyLightIntensity);
-			}
-			else
-			{
-				SphereReflectionComponent->Brightness = Profile.SkyLightIntensity;
-			}
+			SetSkyBrightness(Profile.SkyLightIntensity);
 			
 			InstancedSkyMaterial->SetScalarParameterValueEditorOnly(IntensityName, Profile.SkyLightIntensity);
 			bSkyChanged = true;
@@ -164,15 +136,7 @@ void FAdvancedPreviewScene::UpdateScene(FPreviewSceneProfile& Profile, bool bUpd
 		if (Texture != EnvironmentTexture)
 		{
 			InstancedSkyMaterial->SetTextureParameterValueEditorOnly(SkyBoxName, EnvironmentTexture);
-			if (bUseSkylight)
-			{
-				SetSkyCubemap(EnvironmentTexture);
-			}
-			else
-			{
-				SphereReflectionComponent->Cubemap = EnvironmentTexture;
-			}
-			
+			SetSkyCubemap(EnvironmentTexture);
 			bSkyChanged = true;
 		}
 		
@@ -188,14 +152,7 @@ void FAdvancedPreviewScene::UpdateScene(FPreviewSceneProfile& Profile, bool bUpd
 			LightDir.Yaw = Profile.LightingRigRotation;
 			SetLightDirection(LightDir);
 			DefaultSettings->Profiles[CurrentProfileIndex].DirectionalLightRotation = LightDir;
-			if (bUseSkylight)
-			{
-				SkyLight->SourceCubemapAngle = Profile.LightingRigRotation;
-			}
-			else
-			{
-				SphereReflectionComponent->SourceCubemapAngle = Profile.LightingRigRotation;
-			}
+			SkyLight->SourceCubemapAngle = Profile.LightingRigRotation;
 			bSkyChanged = true;
 		}
 	}		
@@ -217,14 +174,7 @@ void FAdvancedPreviewScene::UpdateScene(FPreviewSceneProfile& Profile, bool bUpd
 	}
 
 	SkyComponent->SetVisibility(Profile.bShowEnvironment, true);
-	if (bUseSkylight)
-	{
-		SkyLight->SetVisibility(Profile.bUseSkyLighting, true);
-	}
-	else
-	{
-		SphereReflectionComponent->SetVisibility(Profile.bShowEnvironment, true);
-	}
+	SkyLight->SetVisibility(Profile.bUseSkyLighting, true);
 	FloorMeshComponent->SetVisibility(Profile.bShowFloor, true);
 
 	bRotateLighting = Profile.bRotateLightingRig;
@@ -272,20 +222,10 @@ void FAdvancedPreviewScene::Tick(float DeltaTime)
 
 	if (!FMath::IsNearlyEqual(PreviousRotation, Profile.LightingRigRotation, 0.05f))
 	{		
-		if (bUseSkylight)
-		{
-			SkyLight->SourceCubemapAngle = Profile.LightingRigRotation;
-			SkyLight->SetCaptureIsDirty();
-			SkyLight->MarkRenderStateDirty();
-			SkyLight->UpdateSkyCaptureContents(PreviewWorld);
-		}
-		else
-		{
-			SphereReflectionComponent->SourceCubemapAngle = Profile.LightingRigRotation;
-			SphereReflectionComponent->MarkDirtyForRecapture();
-			SphereReflectionComponent->MarkRenderStateDirty();
-			UReflectionCaptureComponent::UpdateReflectionCaptureContents(PreviewWorld);
-		}
+		SkyLight->SourceCubemapAngle = Profile.LightingRigRotation;
+		SkyLight->SetCaptureIsDirty();
+		SkyLight->MarkRenderStateDirty();
+		SkyLight->UpdateSkyCaptureContents(PreviewWorld);
 
 		InstancedSkyMaterial->SetScalarParameterValueEditorOnly(FName("CubemapRotation"), Profile.LightingRigRotation / 360.0f);
 		InstancedSkyMaterial->PostEditChange();
@@ -299,19 +239,9 @@ void FAdvancedPreviewScene::Tick(float DeltaTime)
 	// Update the sky every tick rather than every mouse move (UpdateScene call)
 	if (bSkyChanged)
 	{
-		if (bUseSkylight)
-		{
-			SkyLight->SetCaptureIsDirty();
-			SkyLight->MarkRenderStateDirty();
-			SkyLight->UpdateSkyCaptureContents(PreviewWorld);
-		}
-		else
-		{
-			SphereReflectionComponent->MarkDirtyForRecapture();
-			SphereReflectionComponent->MarkRenderStateDirty();
-			UReflectionCaptureComponent::UpdateReflectionCaptureContents(PreviewWorld);
-		}
-
+		SkyLight->SetCaptureIsDirty();
+		SkyLight->MarkRenderStateDirty();
+		SkyLight->UpdateSkyCaptureContents(PreviewWorld);
 
 		InstancedSkyMaterial->PostEditChange();
 		bSkyChanged = false;
@@ -405,10 +335,6 @@ void FAdvancedPreviewScene::SetEnvironmentVisibility(const bool bVisible, const
 	{
 		// Otherwise set visibility directly on the component
 		SkyComponent->SetVisibility(bVisible);
-		if (!bUseSkylight)
-		{
-			SphereReflectionComponent->SetVisibility(bVisible);
-		}
 	}
 }
 
diff --git a/Engine/Source/Editor/AdvancedPreviewScene/Public/AdvancedPreviewScene.h b/Engine/Source/Editor/AdvancedPreviewScene/Public/AdvancedPreviewScene.h
--- a/Engine/Source/Editor/AdvancedPreviewScene/Public/AdvancedPreviewScene.h
+++ b/Engine/Source/Editor/AdvancedPreviewScene/Public/AdvancedPreviewScene.h
@@ -73,7 +73,6 @@ protected:
 
 protected:
 	UStaticMeshComponent* SkyComponent;
-	USphereReflectionCaptureComponent* SphereReflectionComponent;
 	UMaterialInstanceConstant* InstancedSkyMaterial;
 	UPostProcessComponent* PostProcessComponent;
 	UStaticMeshComponent* FloorMeshComponent;
@@ -86,8 +85,6 @@ protected:
 	bool bSkyChanged;
 	bool bPostProcessing;
 
-	bool bUseSkylight;
-
 	int32 CurrentProfileIndex;
 
 	/** Command list for input handling */
