Add functionality to Preview Scene Settings that allows the user to hide Environment Cubemap without disabling lighting
#fix Add bUseImageBasedLighting flag to preview scene, so new behaviour when hiding the enviroment is that IBL is still on
#misc Fixed up comment and readded missing code from old merge

diff --git a/Engine/Source/Editor/AdvancedPreviewScene/Private/AdvancedPreviewScene.cpp b/Engine/Source/Editor/AdvancedPreviewScene/Private/AdvancedPreviewScene.cpp
--- a/Engine/Source/Editor/AdvancedPreviewScene/Private/AdvancedPreviewScene.cpp
+++ b/Engine/Source/Editor/AdvancedPreviewScene/Private/AdvancedPreviewScene.cpp
@@ -41,9 +41,28 @@ FAdvancedPreviewScene::FAdvancedPreviewScene(ConstructionValues CVS, float InFlo
 
 	const FTransform Transform(FRotator(0, 0, 0), FVector(0, 0, 0), FVector(1));
 
-	// Set up sky light using the set cube map texture, reusing the sky light from PreviewScene class
-	SetSkyCubemap(Profile.EnvironmentCubeMap.Get());
-	SetSkyBrightness(Profile.SkyLightIntensity);
+	static const auto CVarSupportStationarySkylight = IConsoleManager::Get().FindTConsoleVariableDataInt(TEXT("r.SupportStationarySkylight"));
+	bUseSkylight = CVarSupportStationarySkylight->GetValueOnAnyThread() != 0;
+
+	if (bUseSkylight)
+	{
+		// Set up sky light using the set cube map texture, reusing the sky light from PreviewScene class
+		SetSkyCubemap(Profile.EnvironmentCubeMap.Get());
+		SetSkyBrightness(Profile.SkyLightIntensity);
+	}
+	else
+	{
+		// Hide the inherited skylight 
+		SkyLight->SetVisibility(false);
+
+		// Setup sphere reflection
+		SphereReflectionComponent = NewObject<USphereReflectionCaptureComponent>();
+		SphereReflectionComponent->Cubemap = Profile.EnvironmentCubeMap.Get();
+		SphereReflectionComponent->ReflectionSourceType = EReflectionSourceType::SpecifiedCubemap;
+		SphereReflectionComponent->Brightness = Profile.SkyLightIntensity;
+		AddComponent(SphereReflectionComponent, Transform);
+		SphereReflectionComponent->UpdateReflectionCaptureContents(PreviewWorld);
+	}
 	
 	// Large scale to prevent sphere from clipping
 	const FTransform SphereTransform(FRotator(0, 0, 0), FVector(0, 0, 0), FVector(2000));
@@ -88,7 +107,6 @@ FAdvancedPreviewScene::FAdvancedPreviewScene(ConstructionValues CVS, float InFlo
 	bRotateLighting = Profile.bRotateLightingRig;
 	CurrentRotationSpeed = Profile.RotationSpeed;
 	bSkyChanged = false;
-	bUseSkylight = true;
 
 	BindCommands();
 }
@@ -194,7 +212,7 @@ void FAdvancedPreviewScene::UpdateScene(FPreviewSceneProfile& Profile, bool bUpd
 	SkyComponent->SetVisibility(Profile.bShowEnvironment, true);
 	if (bUseSkylight)
 	{
-		SkyLight->SetVisibility(Profile.bShowEnvironment, true);
+		SkyLight->SetVisibility(Profile.bUseSkyLighting, true);
 	}
 	else
 	{
@@ -360,7 +378,7 @@ void FAdvancedPreviewScene::SetFloorVisibility(const bool bVisible, const bool b
 	}
 	else
 	{
-		// Otherwise set visiblity directly on the component
+		// Otherwise set visibility directly on the component
 		FloorMeshComponent->SetVisibility(bVisible ? DefaultSettings->Profiles[CurrentProfileIndex].bShowFloor : bVisible);
 	}
 }
@@ -379,14 +397,10 @@ void FAdvancedPreviewScene::SetEnvironmentVisibility(const bool bVisible, const
 	else
 	{
 		// Otherwise set visibility directly on the component
-		SkyComponent->SetVisibility(bVisible ? DefaultSettings->Profiles[CurrentProfileIndex].bShowEnvironment : bVisible);
-		if (bUseSkylight)
-		{
-			SkyLight->SetVisibility(bVisible ? DefaultSettings->Profiles[CurrentProfileIndex].bShowEnvironment : bVisible);
-		}
-		else
+		SkyComponent->SetVisibility(bVisible);
+		if (!bUseSkylight)
 		{
-			SphereReflectionComponent->SetVisibility(bVisible ? DefaultSettings->Profiles[CurrentProfileIndex].bShowEnvironment : bVisible);
+			SphereReflectionComponent->SetVisibility(bVisible);
 		}
 	}
 }
@@ -472,7 +486,7 @@ void FAdvancedPreviewScene::OnAssetViewerSettingsRefresh(const FName& InProperty
 		const bool bNameNone = InPropertyName == NAME_None;
 
 		const bool bUpdateEnvironment = (InPropertyName == GET_MEMBER_NAME_CHECKED(FPreviewSceneProfile, EnvironmentCubeMap)) || (InPropertyName == GET_MEMBER_NAME_CHECKED(FPreviewSceneProfile, LightingRigRotation) || (InPropertyName == GET_MEMBER_NAME_CHECKED(UAssetViewerSettings, Profiles)));
-		const bool bUpdateSkyLight = bUpdateEnvironment || (InPropertyName == GET_MEMBER_NAME_CHECKED(FPreviewSceneProfile, SkyLightIntensity) || (InPropertyName == GET_MEMBER_NAME_CHECKED(UAssetViewerSettings, Profiles)));
+		const bool bUpdateSkyLight = bUpdateEnvironment || (InPropertyName == GET_MEMBER_NAME_CHECKED(FPreviewSceneProfile, SkyLightIntensity) || InPropertyName == GET_MEMBER_NAME_CHECKED(FPreviewSceneProfile, bUseSkyLighting) || (InPropertyName == GET_MEMBER_NAME_CHECKED(UAssetViewerSettings, Profiles)));
 		const bool bUpdateDirectionalLight = (InPropertyName == GET_MEMBER_NAME_CHECKED(FPreviewSceneProfile, DirectionalLightIntensity)) || (InPropertyName == GET_MEMBER_NAME_CHECKED(FPreviewSceneProfile, DirectionalLightColor));
 		const bool bUpdatePostProcessing = (InPropertyName == GET_MEMBER_NAME_CHECKED(FPreviewSceneProfile, PostProcessingSettings)) || (InPropertyName == GET_MEMBER_NAME_CHECKED(FPreviewSceneProfile, bPostProcessingEnabled));
 
diff --git a/Engine/Source/Editor/AdvancedPreviewScene/Public/AssetViewerSettings.h b/Engine/Source/Editor/AdvancedPreviewScene/Public/AssetViewerSettings.h
--- a/Engine/Source/Editor/AdvancedPreviewScene/Public/AssetViewerSettings.h
+++ b/Engine/Source/Editor/AdvancedPreviewScene/Public/AssetViewerSettings.h
@@ -23,6 +23,7 @@ struct FPreviewSceneProfile
 	FPreviewSceneProfile()
 	{		
 		bSharedProfile = false;
+		bUseSkyLighting = true;
 		bShowFloor = true;
 		bShowEnvironment = true;
 		bRotateLightingRig = false;
@@ -48,6 +49,10 @@ struct FPreviewSceneProfile
 	UPROPERTY(EditAnywhere, config, Category = Profile)
 	bool bSharedProfile;
 
+	/** Whether or not image based lighting is enabled for the environment cube map */
+	UPROPERTY(EditAnywhere, AdvancedDisplay, config, Category = Lighting)
+	bool bUseSkyLighting;
+
 	/** Manually set the directional light intensity (0.0 - 20.0) */
 	UPROPERTY(EditAnywhere, config, Category = Lighting, meta = (UIMin = "0.0", UIMax = "20.0"))
 	float DirectionalLightIntensity;
