Sky sphere template asset (BP_Sky_Sphere.uasset) now has bIsVisibleInRayTracing set to false so we do not need to exclude all unlit geometries out of the ray tracing acceleration structure
fyi patrick.kelly, yuriy.odonnell, sebastien.hillaire, charles.derousiers

diff --git a/Engine/Source/Editor/AdvancedPreviewScene/Private/AdvancedPreviewScene.cpp b/Engine/Source/Editor/AdvancedPreviewScene/Private/AdvancedPreviewScene.cpp
--- a/Engine/Source/Editor/AdvancedPreviewScene/Private/AdvancedPreviewScene.cpp
+++ b/Engine/Source/Editor/AdvancedPreviewScene/Private/AdvancedPreviewScene.cpp
@@ -53,6 +53,7 @@ FAdvancedPreviewScene::FAdvancedPreviewScene(ConstructionValues CVS, float InFlo
 	check(SkySphere);
 	SkyComponent->SetStaticMesh(SkySphere);
 	SkyComponent->SetCollisionEnabled(ECollisionEnabled::NoCollision);
+	SkyComponent->bVisibleInRayTracing = false;
 
 	UMaterial* SkyMaterial = LoadObject<UMaterial>(NULL, TEXT("/Engine/EditorMaterials/AssetViewer/M_SkyBox.M_SkyBox"), NULL, LOAD_None, NULL);
 	check(SkyMaterial);
