Adding a Preview Scene Settings Profiles results in incorrect incrementing
#fix take number of profiles-1 instead of just number of profiles when generating profile name

diff --git a/Engine/Source/Editor/AdvancedPreviewScene/Private/SAdvancedPreviewDetailsTab.cpp b/Engine/Source/Editor/AdvancedPreviewScene/Private/SAdvancedPreviewDetailsTab.cpp
--- a/Engine/Source/Editor/AdvancedPreviewScene/Private/SAdvancedPreviewDetailsTab.cpp
+++ b/Engine/Source/Editor/AdvancedPreviewScene/Private/SAdvancedPreviewDetailsTab.cpp
@@ -180,7 +180,7 @@ FReply SAdvancedPreviewDetailsTab::AddProfileButtonClick()
 	
 	// Try to create a valid profile name when one is added
 	bool bFoundValidName = false;
-	int ProfileAppendNum = DefaultSettings->Profiles.Num();
+	int32 ProfileAppendNum = FMath::Max(0, DefaultSettings->Profiles.Num() - 1);
 	FString NewProfileName;
 	while (!bFoundValidName)
 	{
