The Undo doesn't work in Preview Scene Set
#fix hacky fix, adding config UProperty, to get the mutable CDO of UAssetViewerSettings transact, as its not in a package/saved

diff --git a/Engine/Source/Editor/AdvancedPreviewScene/Public/AssetViewerSettings.h b/Engine/Source/Editor/AdvancedPreviewScene/Public/AssetViewerSettings.h
--- a/Engine/Source/Editor/AdvancedPreviewScene/Public/AssetViewerSettings.h
+++ b/Engine/Source/Editor/AdvancedPreviewScene/Public/AssetViewerSettings.h
@@ -192,7 +192,7 @@ public:
 	/** Collection of scene profiles */
 	UPROPERTY(EditAnywhere, transient, Category = Settings, meta=(ShowOnlyInnerProperties))
 	TArray<FPreviewSceneProfile> Profiles;
-	
+
 	/** Cached value to determine whether or not a profile was added or removed */
 	int32 NumProfiles;
 protected:
@@ -202,4 +202,8 @@ protected:
 	FOnAssetViewerProfileAddRemovedEvent OnAssetViewerProfileAddRemovedEvent;
 
 	FOnAssetViewerSettingsPostUndo OnAssetViewerSettingsPostUndoEvent;
+
+	// This will enforce mutable CDO of UAssetViewerSettings transacted
+	UPROPERTY(Config)
+	bool bFakeConfigValue_HACK;
 };
