Fixed some warnings when trying to save a config that is read only.

The dataprep one doesn't have to save if it can't. It recreate the setup on the next boot.
The collision one is in the project setting where there is a visual warning if the ini file is in a read only mode.

diff --git a/Engine/Plugins/Enterprise/DataprepEditor/Source/DataprepEditor/Private/Widgets/SDataprepEditorViewport.cpp b/Engine/Plugins/Enterprise/DataprepEditor/Source/DataprepEditor/Private/Widgets/SDataprepEditorViewport.cpp
--- a/Engine/Plugins/Enterprise/DataprepEditor/Source/DataprepEditor/Private/Widgets/SDataprepEditorViewport.cpp
+++ b/Engine/Plugins/Enterprise/DataprepEditor/Source/DataprepEditor/Private/Widgets/SDataprepEditorViewport.cpp
@@ -1219,7 +1219,8 @@ void SDataprepEditorViewport::LoadDefaultSettings()
 		AssetViewerProfileIndex = DefaultSettings->Profiles.Num();
 
 		DefaultSettings->Profiles.Add( Profile );
-		DefaultSettings->Save();
+		const bool bWarnIfFail = false;
+		DefaultSettings->Save(bWarnIfFail);
 	}
 
 	// Update the profile with the settings for the project
diff --git a/Engine/Source/Editor/AdvancedPreviewScene/Private/AssetViewerSettings.cpp b/Engine/Source/Editor/AdvancedPreviewScene/Private/AssetViewerSettings.cpp
--- a/Engine/Source/Editor/AdvancedPreviewScene/Private/AssetViewerSettings.cpp
+++ b/Engine/Source/Editor/AdvancedPreviewScene/Private/AssetViewerSettings.cpp
@@ -44,7 +44,7 @@ UAssetViewerSettings* UAssetViewerSettings::Get()
 	return DefaultSettings;
 }
 
-void UAssetViewerSettings::Save()
+void UAssetViewerSettings::Save(bool bWarnIfFail)
 {
 	ULocalProfiles* LocalProfilesObject = GetMutableDefault<ULocalProfiles>();
 	USharedProfiles* SharedProfilesObject = GetMutableDefault<USharedProfiles>();
@@ -71,7 +71,7 @@ void UAssetViewerSettings::Save()
 	LocalProfilesObject->SaveConfig();
 
 	SharedProfilesObject->SaveConfig();
-	SharedProfilesObject->TryUpdateDefaultConfigFile();
+	SharedProfilesObject->TryUpdateDefaultConfigFile(FString(), bWarnIfFail);
 }
 
 void UAssetViewerSettings::PostEditChangeProperty(struct FPropertyChangedEvent& PropertyChangedEvent)
diff --git a/Engine/Source/Editor/AdvancedPreviewScene/Public/AssetViewerSettings.h b/Engine/Source/Editor/AdvancedPreviewScene/Public/AssetViewerSettings.h
--- a/Engine/Source/Editor/AdvancedPreviewScene/Public/AssetViewerSettings.h
+++ b/Engine/Source/Editor/AdvancedPreviewScene/Public/AssetViewerSettings.h
@@ -167,8 +167,11 @@ public:
 
 	static UAssetViewerSettings* Get();
 	
-	/** Saves the config data out to the ini files */
-	void Save();
+	/**
+	 * Saves the config data out to the ini files
+	 * @param bWarnIfFail Should we log a warning if a ini file couldn't be saved.
+	 */
+	void Save(bool bWarnIfFail = true);
 
 	DECLARE_EVENT_OneParam(UAssetViewerSettings, FOnAssetViewerSettingsChangedEvent, const FName&);
 	FOnAssetViewerSettingsChangedEvent& OnAssetViewerSettingsChanged() { return OnAssetViewerSettingsChangedEvent; }
diff --git a/Engine/Source/Editor/DetailCustomizations/Private/CollisionProfileDetails.cpp b/Engine/Source/Editor/DetailCustomizations/Private/CollisionProfileDetails.cpp
--- a/Engine/Source/Editor/DetailCustomizations/Private/CollisionProfileDetails.cpp
+++ b/Engine/Source/Editor/DetailCustomizations/Private/CollisionProfileDetails.cpp
@@ -1784,7 +1784,9 @@ void FCollisionProfileDetails::UpdateChannel(bool bTraceType)
 void FCollisionProfileDetails::UpdateProfile()
 {
 	CollisionProfile->LoadProfileConfig(true);
-	CollisionProfile->TryUpdateDefaultConfigFile();
+	const FString SpecificFileLocation;
+	const bool bWarnIfFail = false;
+	CollisionProfile->TryUpdateDefaultConfigFile(SpecificFileLocation, bWarnIfFail);
 	SavedData.Save(CollisionProfile);
 
 	RefreshProfileList();
