AdvancedPreviewScene: Allow to select things below the ground

diff --git a/Engine/Source/Editor/AdvancedPreviewScene/Private/AdvancedPreviewScene.cpp b/Engine/Source/Editor/AdvancedPreviewScene/Private/AdvancedPreviewScene.cpp
--- a/Engine/Source/Editor/AdvancedPreviewScene/Private/AdvancedPreviewScene.cpp
+++ b/Engine/Source/Editor/AdvancedPreviewScene/Private/AdvancedPreviewScene.cpp
@@ -78,6 +78,7 @@ FAdvancedPreviewScene::FAdvancedPreviewScene(ConstructionValues CVS, float InFlo
 	check(FloorMesh);
 	FloorMeshComponent = NewObject<UStaticMeshComponent>(GetTransientPackage());
 	FloorMeshComponent->SetStaticMesh(FloorMesh);
+	FloorMeshComponent->bSelectable = false;
 
 	FTransform FloorTransform(FRotator(0, 0, 0), FVector(0, 0, -(InFloorOffset)), FVector(4.0f, 4.0f, 1.0f ));
 	AddComponent(FloorMeshComponent, FloorTransform);	
