Disable shadows on preview scene sky

PR Description: "In the AdvancedPreviewScene, the skysphere appears to be casting shadows. This causes the directional light to not affect the current asset and prevents the user from properly seeing their changes. This is most obvious with the material editor as well as the static mesh editor, where the asset would appear dark due to only sky lighting being present. I have reproduced this when using Lumen with hardware ray tracing and it appears to be caused by Lumen tracing against backfaces. I have not attempted to reproduce this without hardware ray tracing due to the fact that the fix is simple and won't affect non hardware ray tracing. The sky didn't have the shadows disabled because previous forms of lighting ignored backfaces. But now, with backfaces being traced against, this seems like a good idea."

PR #9965

diff --git a/Engine/Source/Editor/AdvancedPreviewScene/Private/AdvancedPreviewScene.cpp b/Engine/Source/Editor/AdvancedPreviewScene/Private/AdvancedPreviewScene.cpp
--- a/Engine/Source/Editor/AdvancedPreviewScene/Private/AdvancedPreviewScene.cpp
+++ b/Engine/Source/Editor/AdvancedPreviewScene/Private/AdvancedPreviewScene.cpp
@@ -53,6 +53,8 @@ FAdvancedPreviewScene::FAdvancedPreviewScene(ConstructionValues CVS, float InFlo
 	check(SkySphere);
 	SkyComponent->SetStaticMesh(SkySphere);
 	SkyComponent->SetCollisionEnabled(ECollisionEnabled::NoCollision);
+	SkyComponent->CastShadow = false;
+	SkyComponent->bCastDynamicShadow = false;
 
 	UMaterial* SkyMaterial = LoadObject<UMaterial>(NULL, TEXT("/Engine/EditorMaterials/AssetViewer/M_SkyBox.M_SkyBox"), NULL, LOAD_None, NULL);
 	check(SkyMaterial);
