StaticMeshEditor: Fix "Vert Color" toggle ignoring "Show Environment" preview flag

diff --git a/Engine/Source/Editor/AdvancedPreviewScene/Private/AdvancedPreviewScene.cpp b/Engine/Source/Editor/AdvancedPreviewScene/Private/AdvancedPreviewScene.cpp
--- a/Engine/Source/Editor/AdvancedPreviewScene/Private/AdvancedPreviewScene.cpp
+++ b/Engine/Source/Editor/AdvancedPreviewScene/Private/AdvancedPreviewScene.cpp
@@ -378,7 +378,7 @@ void FAdvancedPreviewScene::SetEnvironmentVisibility(const bool bVisible, const
 	else
 	{
 		// Otherwise set visibility directly on the component
-		SkyComponent->SetVisibility(bVisible);
+		SkyComponent->SetVisibility(bVisible ? DefaultSettings->Profiles[CurrentProfileIndex].bShowEnvironment : bVisible);
 	}
 }
 
