Fixes crash when closing skeletal mesh editor while painting weights.

diff --git a/Engine/Source/Editor/AdvancedPreviewScene/Private/SAdvancedPreviewDetailsTab.cpp b/Engine/Source/Editor/AdvancedPreviewScene/Private/SAdvancedPreviewDetailsTab.cpp
--- a/Engine/Source/Editor/AdvancedPreviewScene/Private/SAdvancedPreviewDetailsTab.cpp
+++ b/Engine/Source/Editor/AdvancedPreviewScene/Private/SAdvancedPreviewDetailsTab.cpp
@@ -276,6 +276,13 @@ FReply SAdvancedPreviewDetailsTab::RemoveOrResetProfileButtonClick()
 
 void SAdvancedPreviewDetailsTab::OnAssetViewerSettingsRefresh(const FName& InPropertyName)
 {
+	if (!PreviewScenePtr.IsValid())
+	{
+		// this callback can fire when the editor is forcibly closed and tool modes revert the active profile
+		// when this happens, the preview scene is null even though this details tab hasn't been destroyed yet (and unregistered this delegate)
+		return;
+	}
+	
 	if (InPropertyName == GET_MEMBER_NAME_CHECKED(FPreviewSceneProfile, ProfileName) || InPropertyName == GET_MEMBER_NAME_CHECKED(FPreviewSceneProfile, bSharedProfile))
 	{
 		Refresh();
