Uninitialized variable in FPreviewSceneProfile

diff --git a/Engine/Source/Editor/AdvancedPreviewScene/Public/AssetViewerSettings.h b/Engine/Source/Editor/AdvancedPreviewScene/Public/AssetViewerSettings.h
--- a/Engine/Source/Editor/AdvancedPreviewScene/Public/AssetViewerSettings.h
+++ b/Engine/Source/Editor/AdvancedPreviewScene/Public/AssetViewerSettings.h
@@ -43,6 +43,7 @@ struct FPreviewSceneProfile
 		DirectionalLightRotation = FRotator(-40.f, -67.5f, 0.f);
 		bEnableToneMapping = true;
 		bShowMeshEdges = false;
+		bShowGrid = false;
 	}
 
 	/** Name to identify the profile */
