Fixes crash when displaying normals on skeletal mesh assets

diff --git a/Engine/Source/Editor/AdvancedPreviewScene/Public/AssetViewerSettings.h b/Engine/Source/Editor/AdvancedPreviewScene/Public/AssetViewerSettings.h
--- a/Engine/Source/Editor/AdvancedPreviewScene/Public/AssetViewerSettings.h
+++ b/Engine/Source/Editor/AdvancedPreviewScene/Public/AssetViewerSettings.h
@@ -155,6 +155,10 @@ struct FPreviewSceneProfile
 
 	void SetShowFlags(FEngineShowFlags& ShowFlags) const
 	{
+		// for reasons I haven't been able to discern exactly, this must be called prior to EnableAdvancedFeatures()
+		// to prevent a crash in the renderer caused by unallocated or missing resources
+		ShowFlags.DisableAdvancedFeatures();
+		
 		if (bPostProcessingEnabled)
 		{
 			ShowFlags.EnableAdvancedFeatures();
