Fixed type mismatch when assigning UObject* to TSoftObjectPtr<UTextureCube>.

diff --git a/Engine/Source/Editor/AdvancedPreviewScene/Public/AssetViewerSettings.h b/Engine/Source/Editor/AdvancedPreviewScene/Public/AssetViewerSettings.h
--- a/Engine/Source/Editor/AdvancedPreviewScene/Public/AssetViewerSettings.h
+++ b/Engine/Source/Editor/AdvancedPreviewScene/Public/AssetViewerSettings.h
@@ -148,9 +148,9 @@ struct FPreviewSceneProfile
 					LoadedObject = Redirector->DestinationObject;
 				}
 
-				EnvironmentCubeMap = LoadedObject;
+				EnvironmentCubeMap = Cast<UTextureCube>(LoadedObject);
 			}
-		}		
+		}
 	}
 
 	void SetShowFlags(FEngineShowFlags& ShowFlags) const
