Fix crash when creating an FAdvancedPreviewScene with bDefaultLight set to false

diff --git a/Engine/Source/Editor/AdvancedPreviewScene/Private/AdvancedPreviewScene.cpp b/Engine/Source/Editor/AdvancedPreviewScene/Private/AdvancedPreviewScene.cpp
--- a/Engine/Source/Editor/AdvancedPreviewScene/Private/AdvancedPreviewScene.cpp
+++ b/Engine/Source/Editor/AdvancedPreviewScene/Private/AdvancedPreviewScene.cpp
@@ -117,7 +117,7 @@ FAdvancedPreviewScene::~FAdvancedPreviewScene()
 
 void FAdvancedPreviewScene::UpdateScene(FPreviewSceneProfile& Profile, bool bUpdateSkyLight /*= true*/, bool bUpdateEnvironment  /*= true*/, bool bUpdatePostProcessing /*= true*/, bool bUpdateDirectionalLight /*= true*/)
 {
-	if (bUpdateSkyLight)
+	if (bUpdateSkyLight && SkyLight)
 	{
 		// Threshold to ensure we only update the intensity if it is going to make a difference
 		if (!FMath::IsNearlyEqual(SkyLight->Intensity, Profile.SkyLightIntensity, 0.05f))
@@ -130,7 +130,7 @@ void FAdvancedPreviewScene::UpdateScene(FPreviewSceneProfile& Profile, bool bUpd
 		}
 	}
 
-	if (bUpdateEnvironment)
+	if (bUpdateEnvironment && SkyLight)
 	{
 		static const FName SkyBoxName("SkyBox");
 		static const FName CubeMapRotationName("CubemapRotation");
@@ -178,7 +178,7 @@ void FAdvancedPreviewScene::UpdateScene(FPreviewSceneProfile& Profile, bool bUpd
 		bPostProcessing = Profile.bPostProcessingEnabled;
 	}
 
-	if (bUpdateDirectionalLight)
+	if (bUpdateDirectionalLight && DirectionalLight)
 	{
 		if (!FMath::IsNearlyEqual(DirectionalLight->Intensity, Profile.DirectionalLightIntensity, 0.05f))
 		{
@@ -188,7 +188,11 @@ void FAdvancedPreviewScene::UpdateScene(FPreviewSceneProfile& Profile, bool bUpd
 	}
 
 	SkyComponent->SetVisibility(Profile.bShowEnvironment, true);
-	SkyLight->SetVisibility(Profile.bUseSkyLighting, true);
+
+	if (SkyLight)
+	{
+		SkyLight->SetVisibility(Profile.bUseSkyLighting, true);
+	}
 	FloorMeshComponent->SetVisibility(Profile.bShowFloor, true);
 
 	bRotateLighting = Profile.bRotateLightingRig;
