Renames default preview scene profiles

diff --git a/Engine/Source/Editor/AdvancedPreviewScene/Private/AssetViewerSettings.cpp b/Engine/Source/Editor/AdvancedPreviewScene/Private/AssetViewerSettings.cpp
--- a/Engine/Source/Editor/AdvancedPreviewScene/Private/AssetViewerSettings.cpp
+++ b/Engine/Source/Editor/AdvancedPreviewScene/Private/AssetViewerSettings.cpp
@@ -5,8 +5,8 @@
 #include "Editor/EditorPerProjectUserSettings.h"
 #include "Editor.h"
 
-FName UDefaultEditorProfiles::DefaultProfileName = FName("Preview");
-FName UDefaultEditorProfiles::EditingProfileName = FName("Edit");
+FName UDefaultEditorProfiles::DefaultProfileName = FName("Epic Headquarters");
+FName UDefaultEditorProfiles::EditingProfileName = FName("Grey Wireframe");
 
 const FPreviewSceneProfile* UDefaultEditorProfiles::GetProfile(const FString& ProfileName)
 {
