Hiding preview scene profiles from viewport menu until custom one is added

diff --git a/Engine/Source/Editor/AdvancedPreviewScene/Private/PreviewProfileController.cpp b/Engine/Source/Editor/AdvancedPreviewScene/Private/PreviewProfileController.cpp
--- a/Engine/Source/Editor/AdvancedPreviewScene/Private/PreviewProfileController.cpp
+++ b/Engine/Source/Editor/AdvancedPreviewScene/Private/PreviewProfileController.cpp
@@ -102,6 +102,22 @@ FString FPreviewProfileController::GetActiveProfile() const
 	return FString();
 }
 
+bool FPreviewProfileController::HasAnyUserProfiles() const
+{
+	if (AssetViewerSettings)
+	{
+		for (const FPreviewSceneProfile& Profile : AssetViewerSettings->Profiles)
+		{
+			if (!Profile.bIsEngineDefaultProfile)
+			{
+				return true;
+			}
+		}
+	}
+	
+	return false;
+}
+
 void FPreviewProfileController::EnsureProfilesStateCoherence() const
 {
 	ensureMsgf(AssetViewerProfileNames.Num() == AssetViewerSettings->Profiles.Num(), TEXT("List of profiles is out of sync with the list of corresponding profile names."));
diff --git a/Engine/Source/Editor/AdvancedPreviewScene/Private/SAdvancedPreviewDetailsTab.cpp b/Engine/Source/Editor/AdvancedPreviewScene/Private/SAdvancedPreviewDetailsTab.cpp
--- a/Engine/Source/Editor/AdvancedPreviewScene/Private/SAdvancedPreviewDetailsTab.cpp
+++ b/Engine/Source/Editor/AdvancedPreviewScene/Private/SAdvancedPreviewDetailsTab.cpp
@@ -201,7 +201,7 @@ FReply SAdvancedPreviewDetailsTab::AddProfileButtonClick()
 	DefaultSettings->Modify();
 
 	// Add new profile to settings instance
-	DefaultSettings->Profiles.AddDefaulted();
+	ProfileIndex = DefaultSettings->Profiles.AddDefaulted();
 	FPreviewSceneProfile& NewProfile = DefaultSettings->Profiles.Last();
 	
 	// Try to create a valid profile name when one is added
@@ -231,6 +231,7 @@ FReply SAdvancedPreviewDetailsTab::AddProfileButtonClick()
 	}
 
 	NewProfile.ProfileName = NewProfileName;
+	PerProjectSettings->AssetViewerProfileIndex = ProfileIndex;
 	DefaultSettings->PostEditChange();
 
 	// Change selection to new profile so the user directly sees the profile that was added
diff --git a/Engine/Source/Editor/AdvancedPreviewScene/Public/PreviewProfileController.h b/Engine/Source/Editor/AdvancedPreviewScene/Public/PreviewProfileController.h
--- a/Engine/Source/Editor/AdvancedPreviewScene/Public/PreviewProfileController.h
+++ b/Engine/Source/Editor/AdvancedPreviewScene/Public/PreviewProfileController.h
@@ -30,6 +30,9 @@ public:
 	/** Returns the preview profiles currently active. */
 	virtual FString GetActiveProfile() const override;
 
+	/** Returns true if user has added one or more of their own profiles */
+	virtual bool HasAnyUserProfiles() const override;
+
 	/** Invoked after the list of available profiles has changed. */
 	virtual FOnPreviewProfileListChanged& OnPreviewProfileListChanged() override { return OnPreviewProfileListChangedDelegate; }
 
diff --git a/Engine/Source/Editor/Persona/Private/SAnimViewportToolBar.cpp b/Engine/Source/Editor/Persona/Private/SAnimViewportToolBar.cpp
--- a/Engine/Source/Editor/Persona/Private/SAnimViewportToolBar.cpp
+++ b/Engine/Source/Editor/Persona/Private/SAnimViewportToolBar.cpp
@@ -331,6 +331,16 @@ void SAnimViewportToolBar::Construct(const FArguments& InArgs, TSharedPtr<class
 		.Padding(ToolbarSlotPadding)
 		[
 			SNew(SPreviewSceneProfileSelector).PreviewProfileController(InArgs._PreviewProfileController)
+			.Visibility_Lambda([this]()
+			{
+				// only show this menu if the user has customized it by adding their own profiles
+				// this behavior was requested by UX to match the behavior of the static mesh editor
+				if (PreviewProfileController.IsValid())
+				{
+					return PreviewProfileController->HasAnyUserProfiles() ? EVisibility::Visible : EVisibility::Hidden;
+				}
+				return EVisibility::Hidden;
+			})
 		]
 		+ SHorizontalBox::Slot()
 		.AutoWidth()
diff --git a/Engine/Source/Editor/UnrealEd/Public/IPreviewProfileController.h b/Engine/Source/Editor/UnrealEd/Public/IPreviewProfileController.h
--- a/Engine/Source/Editor/UnrealEd/Public/IPreviewProfileController.h
+++ b/Engine/Source/Editor/UnrealEd/Public/IPreviewProfileController.h
@@ -24,6 +24,9 @@ public:
 	/** Returns the preview profiles currently active. */
 	virtual FString GetActiveProfile() const = 0;
 
+	/** Returns true if user has added one or more of their own profiles */
+	virtual bool HasAnyUserProfiles() const = 0;
+
 	/** Invoked after the list of available profiles has changed. */
 	DECLARE_EVENT(IPreviewProfileController, FOnPreviewProfileListChanged);
 	virtual FOnPreviewProfileListChanged& OnPreviewProfileListChanged() = 0;
