Ensure preview scene settings are applied after creation

diff --git a/Engine/Source/Editor/AdvancedPreviewScene/Private/AdvancedPreviewScene.cpp b/Engine/Source/Editor/AdvancedPreviewScene/Private/AdvancedPreviewScene.cpp
--- a/Engine/Source/Editor/AdvancedPreviewScene/Private/AdvancedPreviewScene.cpp
+++ b/Engine/Source/Editor/AdvancedPreviewScene/Private/AdvancedPreviewScene.cpp
@@ -101,6 +101,9 @@ FAdvancedPreviewScene::FAdvancedPreviewScene(ConstructionValues CVS, float InFlo
 
 	PreviousRotation = Profile.LightingRigRotation;
 	UILightingRigRotationDelta = 0.0f;
+
+	// Ensure visibility properties etc are set (note that UpdateScene doesn't get called continuously)
+	UpdateScene(Profile);
 }
 
 FAdvancedPreviewScene::~FAdvancedPreviewScene()
