Add a new FPreviewSceneProfile to UDefaultEditorProfiles: "Grey Ambient", a low-contrast environment suitable for asset editor preview viewports

#virtualized

diff --git a/Engine/Source/Editor/AdvancedPreviewScene/Private/AssetViewerSettings.cpp b/Engine/Source/Editor/AdvancedPreviewScene/Private/AssetViewerSettings.cpp
--- a/Engine/Source/Editor/AdvancedPreviewScene/Private/AssetViewerSettings.cpp
+++ b/Engine/Source/Editor/AdvancedPreviewScene/Private/AssetViewerSettings.cpp
@@ -7,6 +7,7 @@
 
 FName UDefaultEditorProfiles::DefaultProfileName = FName("Epic Headquarters");
 FName UDefaultEditorProfiles::EditingProfileName = FName("Grey Wireframe");
+FName UDefaultEditorProfiles::GreyAmbientProfileName = FName("Grey Ambient");
 
 const FPreviewSceneProfile* UDefaultEditorProfiles::GetProfile(const FString& ProfileName)
 {
diff --git a/Engine/Source/Editor/AdvancedPreviewScene/Public/AssetViewerSettings.h b/Engine/Source/Editor/AdvancedPreviewScene/Public/AssetViewerSettings.h
--- a/Engine/Source/Editor/AdvancedPreviewScene/Public/AssetViewerSettings.h
+++ b/Engine/Source/Editor/AdvancedPreviewScene/Public/AssetViewerSettings.h
@@ -203,12 +203,29 @@ public:
 		EditingProfile.bShowMeshEdges = true;
 		EditingProfile.bEnableToneMapping = false;
 		Profiles.Add(EditingProfile);
+
+		FPreviewSceneProfile GreyAmbientProfile;
+		GreyAmbientProfile.ProfileName = GreyAmbientProfileName.ToString();
+		GreyAmbientProfile.bIsEngineDefaultProfile = true;
+		GreyAmbientProfile.bSharedProfile = true;
+		GreyAmbientProfile.bShowEnvironment = true;
+		GreyAmbientProfile.bShowFloor = true;
+		GreyAmbientProfile.bShowGrid = true;
+		GreyAmbientProfile.bUseSkyLighting = true;
+		GreyAmbientProfile.bPostProcessingEnabled = false;
+		GreyAmbientProfile.bShowMeshEdges = false;
+		GreyAmbientProfile.bEnableToneMapping = false;
+		GreyAmbientProfile.DirectionalLightIntensity = 4.0;
+		GreyAmbientProfile.SkyLightIntensity = 2.0;
+		GreyAmbientProfile.EnvironmentCubeMapPath = TEXT("/Engine/EditorMaterials/AssetViewer/T_GreyAmbient");
+		Profiles.Add(GreyAmbientProfile);
 	}
 
 	const FPreviewSceneProfile* GetProfile(const FString& ProfileName);
 
 	static FName DefaultProfileName;
 	static FName EditingProfileName;
+	static FName GreyAmbientProfileName;
 	
 	/** Collection of default engine-provided profiles used in various editing environments*/
 	UPROPERTY()
