[Deprecation] Remove deprecated function from advanced preview scene

diff --git a/Engine/Source/Editor/AdvancedPreviewScene/Private/AdvancedPreviewScene.cpp b/Engine/Source/Editor/AdvancedPreviewScene/Private/AdvancedPreviewScene.cpp
--- a/Engine/Source/Editor/AdvancedPreviewScene/Private/AdvancedPreviewScene.cpp
+++ b/Engine/Source/Editor/AdvancedPreviewScene/Private/AdvancedPreviewScene.cpp
@@ -315,19 +315,6 @@ const bool FAdvancedPreviewScene::HandleViewportInput(FViewport* InViewport, FIn
 	return bResult;
 }
 
-PRAGMA_DISABLE_DEPRECATION_WARNINGS
-const bool FAdvancedPreviewScene::HandleViewportInput(FViewport* InViewport, int32 ControllerId, FKey Key, float Delta, float DeltaTime, int32 NumSamples, bool bGamepad)
-{
-	return HandleViewportInput(InViewport, FInputDeviceId::CreateFromInternalId(ControllerId), Key, Delta, DeltaTime, NumSamples, bGamepad);
-}
-
-const bool FAdvancedPreviewScene::HandleInputKey(FViewport* InViewport, int32 ControllerId, FKey Key, EInputEvent Event, float AmountDepressed, bool Gamepad)
-{
-	FInputKeyEventArgs Args(InViewport, ControllerId, Key, Event, AmountDepressed, false);
-	return HandleInputKey(Args);
-}
-PRAGMA_ENABLE_DEPRECATION_WARNINGS
-
 const bool FAdvancedPreviewScene::HandleInputKey(const FInputKeyEventArgs& EventArgs)
 {
 	if (EventArgs.Event == IE_Pressed)
diff --git a/Engine/Source/Editor/AdvancedPreviewScene/Public/AdvancedPreviewScene.h b/Engine/Source/Editor/AdvancedPreviewScene/Public/AdvancedPreviewScene.h
--- a/Engine/Source/Editor/AdvancedPreviewScene/Public/AdvancedPreviewScene.h
+++ b/Engine/Source/Editor/AdvancedPreviewScene/Public/AdvancedPreviewScene.h
@@ -41,14 +41,7 @@ public:
 	virtual TStatId GetStatId() const override;
 	/* End FTickableEditorObject */
 
-	UE_DEPRECATED(5.1, "This version of HandleInputKey is deprecated. Please use the version that takes EventArgs instead.")
-	const bool HandleViewportInput(FViewport* InViewport, int32 ControllerId, FKey Key, float Delta, float DeltaTime, int32 NumSamples, bool bGamepad);
-	
 	const bool HandleViewportInput(FViewport* InViewport, FInputDeviceId DeviceId, FKey Key, float Delta, float DeltaTime, int32 NumSamples, bool bGamepad);
-
-	UE_DEPRECATED(5.1, "This version of HandleInputKey is deprecated. Please use the version that takes EventArgs instead.")
-	const bool HandleInputKey(FViewport* InViewport, int32 ControllerId, FKey Key, EInputEvent Event, float AmountDepressed, bool Gamepad);
-
 	const bool HandleInputKey(const FInputKeyEventArgs& EventArgs);
 
 	void SetSkyRotation(const float SkyRotation);
