Refactor AdvancedPreviewScene settings behavior to be more consistent and set the stage for per-editor preview scene settings.
- FAdvancedPreviewScene owns its own "profile changed" event. A default implementation syncronizes the profile settings with the provided viewport & its showflags.
- Remove many direct references to UAssetViewerSettings::OnAssetViewerSettingsChanged(), replacing with UE::AdvancedPreviewScene::BindDefaultOnSettingsChangedHandler().

diff --git a/Engine/Plugins/ChaosClothAssetEditor/Source/ChaosClothAssetEditor/Private/ChaosClothAsset/ClothEditor3DViewportClient.cpp b/Engine/Plugins/ChaosClothAssetEditor/Source/ChaosClothAssetEditor/Private/ChaosClothAsset/ClothEditor3DViewportClient.cpp
--- a/Engine/Plugins/ChaosClothAssetEditor/Source/ChaosClothAssetEditor/Private/ChaosClothAsset/ClothEditor3DViewportClient.cpp
+++ b/Engine/Plugins/ChaosClothAssetEditor/Source/ChaosClothAssetEditor/Private/ChaosClothAsset/ClothEditor3DViewportClient.cpp
@@ -63,9 +63,6 @@ FChaosClothAssetEditor3DViewportClient::FChaosClothAssetEditor3DViewportClient(F
 
 	InPreviewScene->SetGizmoDataBinder(DataBinder);
 
-	// Set correct flags according to current profile settings
-	SetAdvancedShowFlagsForScene(UAssetViewerSettings::Get()->Profiles[GetMutableDefault<UEditorPerProjectUserSettings>()->AssetViewerProfileIndex].bPostProcessingEnabled);
-
 	//
 	// Input behaviors
 	//
@@ -83,10 +80,6 @@ FChaosClothAssetEditor3DViewportClient::FChaosClothAssetEditor3DViewportClient(F
 
 void FChaosClothAssetEditor3DViewportClient::RegisterDelegates()
 {
-	// Remove any existing delegate in case this function is called twice
-	UAssetViewerSettings::Get()->OnAssetViewerSettingsChanged().RemoveAll(this);
-	UAssetViewerSettings::Get()->OnAssetViewerSettingsChanged().AddSP(this, &FChaosClothAssetEditor3DViewportClient::OnAssetViewerSettingsChanged);
-
 	USelection* const SelectedComponents = ModeTools->GetSelectedComponents();
 	SelectedComponents->SelectionChangedEvent.RemoveAll(this);
 	SelectedComponents->SelectionChangedEvent.AddSP(this, &FChaosClothAssetEditor3DViewportClient::ComponentSelectionChanged);
@@ -95,8 +88,7 @@ void FChaosClothAssetEditor3DViewportClient::RegisterDelegates()
 FChaosClothAssetEditor3DViewportClient::~FChaosClothAssetEditor3DViewportClient()
 {
 	DeleteViewportGizmo();
-
-	UAssetViewerSettings::Get()->OnAssetViewerSettingsChanged().RemoveAll(this);
+	
 	if (USelection* const SelectedComponents = ModeTools->GetSelectedComponents())
 	{
 		SelectedComponents->SelectionChangedEvent.RemoveAll(this);
@@ -449,33 +441,4 @@ void FChaosClothAssetEditor3DViewportClient::ComponentSelectionChanged(UObject*
 	}
 }
 
-
-void FChaosClothAssetEditor3DViewportClient::OnAssetViewerSettingsChanged(const FName& InPropertyName)
-{
-	if (InPropertyName == GET_MEMBER_NAME_CHECKED(FPreviewSceneProfile, bPostProcessingEnabled) || InPropertyName == NAME_None)
-	{
-		const UAssetViewerSettings* const Settings = UAssetViewerSettings::Get();
-		const TSharedPtr<const FChaosClothPreviewScene> PinnedClothPreviewScene = ClothPreviewScene.Pin();
-		if (Settings && PinnedClothPreviewScene)
-		{
-			const int32 ProfileIndex = PinnedClothPreviewScene->GetCurrentProfileIndex();
-			if (Settings->Profiles.IsValidIndex(ProfileIndex))
-			{
-				SetAdvancedShowFlagsForScene(Settings->Profiles[ProfileIndex].bPostProcessingEnabled);
-			}
-		}
-	}
-}
-
-void FChaosClothAssetEditor3DViewportClient::SetAdvancedShowFlagsForScene(const bool bAdvancedShowFlags)
-{
-	if (bAdvancedShowFlags)
-	{
-		EngineShowFlags.EnableAdvancedFeatures();
-	}
-	else
-	{
-		EngineShowFlags.DisableAdvancedFeatures();
-	}
-}
 } // namespace UE::Chaos::ClothAsset
diff --git a/Engine/Plugins/ChaosClothAssetEditor/Source/ChaosClothAssetEditor/Private/ChaosClothAsset/ClothEditor3DViewportClient.h b/Engine/Plugins/ChaosClothAssetEditor/Source/ChaosClothAssetEditor/Private/ChaosClothAsset/ClothEditor3DViewportClient.h
--- a/Engine/Plugins/ChaosClothAssetEditor/Source/ChaosClothAssetEditor/Private/ChaosClothAsset/ClothEditor3DViewportClient.h
+++ b/Engine/Plugins/ChaosClothAssetEditor/Source/ChaosClothAssetEditor/Private/ChaosClothAsset/ClothEditor3DViewportClient.h
@@ -114,9 +114,6 @@ private:
 	// IInputBehaviorSource
 	virtual const UInputBehaviorSet* GetInputBehaviors() const override;
 
-	void OnAssetViewerSettingsChanged(const FName& InPropertyName);
-	void SetAdvancedShowFlagsForScene(const bool bAdvancedShowFlags);
-
 	void ComponentSelectionChanged(UObject* NewSelection);
 
 	// Update the selected components based on hitproxy
diff --git a/Engine/Plugins/ChaosClothAssetEditor/Source/ChaosClothAssetEditor/Private/ChaosClothAsset/SClothEditor3DViewport.cpp b/Engine/Plugins/ChaosClothAssetEditor/Source/ChaosClothAssetEditor/Private/ChaosClothAsset/SClothEditor3DViewport.cpp
--- a/Engine/Plugins/ChaosClothAssetEditor/Source/ChaosClothAssetEditor/Private/ChaosClothAsset/SClothEditor3DViewport.cpp
+++ b/Engine/Plugins/ChaosClothAssetEditor/Source/ChaosClothAssetEditor/Private/ChaosClothAsset/SClothEditor3DViewport.cpp
@@ -2,6 +2,7 @@
 
 #include "ChaosClothAsset/SClothEditor3DViewport.h"
 
+#include "AdvancedPreviewSceneMenus.h"
 #include "ClothEditorViewportToolbarSections.h"
 #include "ChaosClothAsset/ClothEditor3DViewportClient.h"
 #include "ChaosClothAsset/SClothEditor3DViewportToolBar.h"
@@ -51,6 +52,11 @@ void SChaosClothAssetEditor3DViewport::Construct(const FArguments& InArgs, const
 			]
 		]
 	];
+	
+	if (TSharedPtr<UE::Chaos::ClothAsset::FChaosClothPreviewScene> PreviewScene = GetPreviewScene().Pin())
+	{
+		UE::AdvancedPreviewScene::BindDefaultOnSettingsChangedHandler(PreviewScene, Client);
+	}
 }
 
 TWeakPtr<UE::Chaos::ClothAsset::FChaosClothPreviewScene> SChaosClothAssetEditor3DViewport::GetPreviewScene()
diff --git a/Engine/Plugins/Experimental/DynamicMaterial/Source/DynamicMaterialEditor/Private/UI/Widgets/Editor/SDMMaterialPreview.cpp b/Engine/Plugins/Experimental/DynamicMaterial/Source/DynamicMaterialEditor/Private/UI/Widgets/Editor/SDMMaterialPreview.cpp
--- a/Engine/Plugins/Experimental/DynamicMaterial/Source/DynamicMaterialEditor/Private/UI/Widgets/Editor/SDMMaterialPreview.cpp
+++ b/Engine/Plugins/Experimental/DynamicMaterial/Source/DynamicMaterialEditor/Private/UI/Widgets/Editor/SDMMaterialPreview.cpp
@@ -3,6 +3,7 @@
 #include "UI/Widgets/Editor/SDMMaterialPreview.h"
 
 #include "AdvancedPreviewScene.h"
+#include "AdvancedPreviewSceneMenus.h"
 #include "AssetViewerSettings.h"
 #include "Components/StaticMeshComponent.h"
 #include "DynamicMaterialEditorCommands.h"
@@ -59,11 +60,6 @@ SDMMaterialPreview::~SDMMaterialPreview()
 		Editor->OnPreviewFeatureLevelChanged().RemoveAll(this);
 	}
 
-	if (UAssetViewerSettings* AssetViewerSettings = UAssetViewerSettings::Get())
-	{
-		AssetViewerSettings->OnAssetViewerSettingsChanged().RemoveAll(this);
-	}
-
 	if (UDynamicMaterialEditorSettings* EditorSettings = UDynamicMaterialEditorSettings::Get())
 	{
 		EditorSettings->GetOnSettingsChanged().RemoveAll(this);
@@ -106,17 +102,7 @@ void SDMMaterialPreview::Construct(const FArguments& InArgs, const TSharedRef<SD
 		SetPreviewType(EditorSettings->PreviewMesh);
 	}
 
-	if (UAssetViewerSettings* AssetViewerSettings = UAssetViewerSettings::Get())
-	{
-		AssetViewerSettings->OnAssetViewerSettingsChanged().AddSP(this, &SDMMaterialPreview::OnAssetViewerSettingsChanged);
-
-		const int32 ProfileIndex = PreviewScene->GetCurrentProfileIndex();
-
-		if (AssetViewerSettings->Profiles.IsValidIndex(ProfileIndex))
-		{
-			PreviewScene->SetEnvironmentVisibility(AssetViewerSettings->Profiles[ProfileIndex].bShowEnvironment, /* Direct */ true);
-		}
-	}
+	UE::AdvancedPreviewScene::BindDefaultOnSettingsChangedHandler(PreviewScene, EditorViewportClient);
 
 	FCoreUObjectDelegates::OnObjectPropertyChanged.AddSP(this, &SDMMaterialPreview::OnPropertyChanged);
 
@@ -446,10 +432,6 @@ void SDMMaterialPreview::OnFeatureLevelChanged(ERHIFeatureLevel::Type InNewFeatu
 	
 }
 
-void SDMMaterialPreview::OnAssetViewerSettingsChanged(const FName& InPropertyName)
-{
-}
-
 TSharedRef<SWidget> SDMMaterialPreview::GenerateToolbarMenu()
 {
 	const FName MenuName = TEXT("DMMaterialPreview");
diff --git a/Engine/Plugins/Experimental/DynamicMaterial/Source/DynamicMaterialEditor/Private/UI/Widgets/Editor/SDMMaterialPreview.h b/Engine/Plugins/Experimental/DynamicMaterial/Source/DynamicMaterialEditor/Private/UI/Widgets/Editor/SDMMaterialPreview.h
--- a/Engine/Plugins/Experimental/DynamicMaterial/Source/DynamicMaterialEditor/Private/UI/Widgets/Editor/SDMMaterialPreview.h
+++ b/Engine/Plugins/Experimental/DynamicMaterial/Source/DynamicMaterialEditor/Private/UI/Widgets/Editor/SDMMaterialPreview.h
@@ -110,8 +110,6 @@ protected:
 
 	void OnFeatureLevelChanged(ERHIFeatureLevel::Type InNewFeatureLevel);
 
-	void OnAssetViewerSettingsChanged(const FName& InPropertyName);
-
 	TSharedRef<SWidget> GenerateToolbarMenu();
 
 	static void AddActionMenu(UToolMenu* InMenu);
diff --git a/Engine/Plugins/FX/Niagara/Source/NiagaraEditor/Private/Widgets/SNiagaraSystemViewport.cpp b/Engine/Plugins/FX/Niagara/Source/NiagaraEditor/Private/Widgets/SNiagaraSystemViewport.cpp
--- a/Engine/Plugins/FX/Niagara/Source/NiagaraEditor/Private/Widgets/SNiagaraSystemViewport.cpp
+++ b/Engine/Plugins/FX/Niagara/Source/NiagaraEditor/Private/Widgets/SNiagaraSystemViewport.cpp
@@ -876,7 +876,7 @@ void SNiagaraSystemViewport::Construct(const FArguments& InArgs, TSharedRef<FNia
 		);
 	}
 	
-	UAssetViewerSettings::Get()->OnAssetViewerSettingsChanged().AddSP(this, &SNiagaraSystemViewport::OnAssetViewerSettingsChanged);
+	UE::AdvancedPreviewScene::BindDefaultOnSettingsChangedHandler(AdvancedPreviewScene, Client);
 }
 
 SNiagaraSystemViewport::~SNiagaraSystemViewport()
@@ -893,8 +893,6 @@ SNiagaraSystemViewport::~SNiagaraSystemViewport()
 	{
 		SystemViewportClient->Viewport = NULL;
 	}
-	
-	UAssetViewerSettings::Get()->OnAssetViewerSettingsChanged().RemoveAll(this);
 }
 
 void SNiagaraSystemViewport::CreateThumbnail(UObject* InScreenShotOwner, TOptional<FGuid> InEmitterToCaptureThumbnailFor)
@@ -1058,11 +1056,6 @@ void SNiagaraSystemViewport::OnScreenShotCaptured(UTexture2D* ScreenShot)
 	EmitterToCaptureThumbnailFor.Reset();
 }
 
-void SNiagaraSystemViewport::OnAssetViewerSettingsChanged(const FName& InPropertyName)
-{
-	RefreshViewport();
-}
-
 void SNiagaraSystemViewport::BindCommands()
 {
 	SEditorViewport::BindCommands();
@@ -1312,8 +1305,8 @@ TSharedPtr<SWidget> SNiagaraSystemViewport::BuildViewportToolbar()
 
 	FToolMenuContext ViewportToolbarContext;
 	{
-		ViewportToolbarContext.AppendCommandList(GetCommandList());
 		ViewportToolbarContext.AppendCommandList(AdvancedPreviewScene->GetCommandList());
+		ViewportToolbarContext.AppendCommandList(GetCommandList());
 
 		// Add the UnrealEd viewport toolbar context.
 		{
diff --git a/Engine/Plugins/FX/Niagara/Source/NiagaraEditor/Private/Widgets/SNiagaraSystemViewport.h b/Engine/Plugins/FX/Niagara/Source/NiagaraEditor/Private/Widgets/SNiagaraSystemViewport.h
--- a/Engine/Plugins/FX/Niagara/Source/NiagaraEditor/Private/Widgets/SNiagaraSystemViewport.h
+++ b/Engine/Plugins/FX/Niagara/Source/NiagaraEditor/Private/Widgets/SNiagaraSystemViewport.h
@@ -123,7 +123,6 @@ protected:
 private:
 	virtual bool IsVisible() const override;
 	void OnScreenShotCaptured(UTexture2D* ScreenShot);
-	void OnAssetViewerSettingsChanged(const FName& InPropertyName);
 
 	/** The parent tab where this viewport resides */
 	TWeakPtr<SDockTab> ParentTab;
diff --git a/Engine/Plugins/Runtime/GeometryCache/Source/GeometryCacheEd/Private/Widgets/SGeometryCacheEditorViewport.cpp b/Engine/Plugins/Runtime/GeometryCache/Source/GeometryCacheEd/Private/Widgets/SGeometryCacheEditorViewport.cpp
--- a/Engine/Plugins/Runtime/GeometryCache/Source/GeometryCacheEd/Private/Widgets/SGeometryCacheEditorViewport.cpp
+++ b/Engine/Plugins/Runtime/GeometryCache/Source/GeometryCacheEd/Private/Widgets/SGeometryCacheEditorViewport.cpp
@@ -3,6 +3,7 @@
 #include "SGeometryCacheEditorViewport.h"
 
 #include "AdvancedPreviewScene.h"
+#include "AdvancedPreviewSceneMenus.h"
 #include "AssetViewerSettings.h"
 #include "Editor/EditorPerProjectUserSettings.h"
 #include "Editor/UnrealEdEngine.h"
@@ -43,9 +44,6 @@ public:
 	virtual bool CanSetWidgetMode(UE::Widget::EWidgetMode NewMode) const override { return false; }
 	virtual bool CanCycleWidgetMode() const override { return false; }
 
-	void OnAssetViewerSettingsChanged(const FName& InPropertyName);
-	void SetAdvancedShowFlagsForScene(const bool bAdvancedShowFlags);
-
 private:
 	TWeakPtr<SGeometryCacheEditorViewport> GeometryCacheEditorViewportPtr;
 };
@@ -78,23 +76,16 @@ FGeometryCacheEditorViewportClient::FGeometryCacheEditorViewportClient(FAdvanced
 	//This seems to be needed to get the correct world time in the preview.
 	SetIsSimulateInEditorViewport(true);
 
-	// Register delegate to update the show flags when the post-processing is turned on or off
-	UAssetViewerSettings::Get()->OnAssetViewerSettingsChanged().AddRaw(this, &FGeometryCacheEditorViewportClient::OnAssetViewerSettingsChanged);
-
 	UEditorPerProjectUserSettings* PerProjectSettings = GetMutableDefault<UEditorPerProjectUserSettings>();
 	UAssetViewerSettings* DefaultSettings = UAssetViewerSettings::Get();
 	PerProjectSettings->AssetViewerProfileIndex = DefaultSettings->Profiles.IsValidIndex(PerProjectSettings->AssetViewerProfileIndex) ? PerProjectSettings->AssetViewerProfileIndex : 0;
 
 	FAdvancedPreviewScene* PreviewSceneCasted = static_cast<FAdvancedPreviewScene*>(PreviewScene);
 	PreviewSceneCasted->SetProfileIndex(PerProjectSettings->AssetViewerProfileIndex);
-
-	// Set correct flags according to current profile settings
-	SetAdvancedShowFlagsForScene(UAssetViewerSettings::Get()->Profiles[PerProjectSettings->AssetViewerProfileIndex].bPostProcessingEnabled);
 }
 
 FGeometryCacheEditorViewportClient::~FGeometryCacheEditorViewportClient()
 {
-	UAssetViewerSettings::Get()->OnAssetViewerSettingsChanged().RemoveAll(this);
 }
 
 FLinearColor FGeometryCacheEditorViewportClient::GetBackgroundColor() const
@@ -125,32 +116,6 @@ void FGeometryCacheEditorViewportClient::Draw(FViewport* InViewport, FCanvas* Ca
 	FEditorViewportClient::Draw(InViewport, Canvas);
 }
 
-void FGeometryCacheEditorViewportClient::OnAssetViewerSettingsChanged(const FName& InPropertyName)
-{
-	if (InPropertyName == GET_MEMBER_NAME_CHECKED(FPreviewSceneProfile, bPostProcessingEnabled) || InPropertyName == NAME_None)
-	{
-		FAdvancedPreviewScene* PreviewSceneCasted = static_cast<FAdvancedPreviewScene*>(PreviewScene);
-		UAssetViewerSettings* Settings = UAssetViewerSettings::Get();
-		const int32 ProfileIndex = PreviewSceneCasted->GetCurrentProfileIndex();
-		if (Settings->Profiles.IsValidIndex(ProfileIndex))
-		{
-			SetAdvancedShowFlagsForScene(Settings->Profiles[ProfileIndex].bPostProcessingEnabled);
-		}
-	}
-}
-
-void FGeometryCacheEditorViewportClient::SetAdvancedShowFlagsForScene(const bool bAdvancedShowFlags)
-{
-	if (bAdvancedShowFlags)
-	{
-		EngineShowFlags.EnableAdvancedFeatures();
-	}
-	else
-	{
-		EngineShowFlags.DisableAdvancedFeatures();
-	}
-}
-
 ////// EDITOR			////////////////
 
 void SGeometryCacheEditorViewport::Construct(const FArguments& InArgs)
@@ -158,9 +123,12 @@ void SGeometryCacheEditorViewport::Construct(const FArguments& InArgs)
 	PreviewGeometryCacheComponent = nullptr;
 
 	AdvancedPreviewScene = MakeShareable(new FAdvancedPreviewScene(FPreviewScene::ConstructionValues()));
-	AdvancedPreviewScene->SetFloorVisibility(false);
 
 	SEditorViewport::Construct(SEditorViewport::FArguments());
+	
+	UE::AdvancedPreviewScene::BindDefaultOnSettingsChangedHandler(AdvancedPreviewScene, Client);
+	
+	AdvancedPreviewScene->SetFloorVisibility(false, /* bDirect */ true);
 }
 
 SGeometryCacheEditorViewport::~SGeometryCacheEditorViewport()
diff --git a/Engine/Plugins/Runtime/HairStrands/Source/HairStrandsEditor/Private/Widgets/SGroomEditorViewport.cpp b/Engine/Plugins/Runtime/HairStrands/Source/HairStrandsEditor/Private/Widgets/SGroomEditorViewport.cpp
--- a/Engine/Plugins/Runtime/HairStrands/Source/HairStrandsEditor/Private/Widgets/SGroomEditorViewport.cpp
+++ b/Engine/Plugins/Runtime/HairStrands/Source/HairStrandsEditor/Private/Widgets/SGroomEditorViewport.cpp
@@ -16,6 +16,7 @@
 #include "ToolMenus.h"
 #include "DrawDebugHelpers.h"
 #include "AdvancedPreviewScene.h"
+#include "AdvancedPreviewSceneMenus.h"
 #include "Widgets/Docking/SDockTab.h"
 #include "GroomComponent.h"
 #include "Components/SkeletalMeshComponent.h"
@@ -42,8 +43,6 @@ public:
 	FGroomEditorViewportClient(FAdvancedPreviewScene& InPreviewScene, const TSharedRef<SGroomEditorViewport>& InGroomEditorViewport);
 	virtual ~FGroomEditorViewportClient() override;
 
-	void OnAssetViewerSettingsChanged(const FName& InPropertyName);
-
 	// FEditorViewportClient interface
 	virtual FLinearColor GetBackgroundColor() const override;
 	virtual void Tick(float DeltaSeconds) override;
@@ -56,8 +55,7 @@ public:
 	virtual bool CanCycleWidgetMode() const override { return false; }
 
 	void SetShowGrid(bool bShowGrid);
-	void SetAdvancedShowFlagsForScene(const bool bAdvancedShowFlags);
-
+	
 	virtual void SetIsSimulateInEditorViewport(bool bInIsSimulateInEditorViewport)override;
 	
 	TWeakPtr<SGroomEditorViewport> GroomEditorViewportPtr;	
@@ -91,49 +89,16 @@ FGroomEditorViewportClient::FGroomEditorViewportClient(FAdvancedPreviewScene& In
 	//This seems to be needed to get the correct world time in the preview.
 	SetIsSimulateInEditorViewport(true);	
 
-	// Register delegate to update the show flags when the post-processing is turned on or off
-	UAssetViewerSettings::Get()->OnAssetViewerSettingsChanged().AddRaw(this, &FGroomEditorViewportClient::OnAssetViewerSettingsChanged);
-
 	UEditorPerProjectUserSettings* PerProjectSettings = GetMutableDefault<UEditorPerProjectUserSettings>();
 	UAssetViewerSettings* DefaultSettings = UAssetViewerSettings::Get();
 	PerProjectSettings->AssetViewerProfileIndex = DefaultSettings->Profiles.IsValidIndex(PerProjectSettings->AssetViewerProfileIndex) ? PerProjectSettings->AssetViewerProfileIndex : 0;
 
 	FAdvancedPreviewScene* PreviewSceneCasted = static_cast<FAdvancedPreviewScene*>(PreviewScene);
 	PreviewSceneCasted->SetProfileIndex(PerProjectSettings->AssetViewerProfileIndex);
-
-	// Set correct flags according to current profile settings
-	SetAdvancedShowFlagsForScene(UAssetViewerSettings::Get()->Profiles[PerProjectSettings->AssetViewerProfileIndex].bPostProcessingEnabled);
 }
 
 FGroomEditorViewportClient::~FGroomEditorViewportClient()
 {
-	UAssetViewerSettings::Get()->OnAssetViewerSettingsChanged().RemoveAll(this);
-}
-
-void FGroomEditorViewportClient::OnAssetViewerSettingsChanged(const FName& InPropertyName)
-{
-	if (InPropertyName == GET_MEMBER_NAME_CHECKED(FPreviewSceneProfile, bPostProcessingEnabled) || InPropertyName == NAME_None)
-	{
-		FAdvancedPreviewScene* PreviewSceneCasted = static_cast<FAdvancedPreviewScene*>(PreviewScene);
-		UAssetViewerSettings* Settings = UAssetViewerSettings::Get();
-		const int32 ProfileIndex = PreviewSceneCasted->GetCurrentProfileIndex();
-		if (Settings->Profiles.IsValidIndex(ProfileIndex))
-		{
-			SetAdvancedShowFlagsForScene(Settings->Profiles[ProfileIndex].bPostProcessingEnabled);
-		}		
-	}
-}
-
-void FGroomEditorViewportClient::SetAdvancedShowFlagsForScene(const bool bAdvancedShowFlags)
-{	
-	if (bAdvancedShowFlags)
-	{
-		EngineShowFlags.EnableAdvancedFeatures();
-	}
-	else
-	{
-		EngineShowFlags.DisableAdvancedFeatures();
-	}
 }
 
 void FGroomEditorViewportClient::Tick(float DeltaSeconds)
@@ -190,6 +155,8 @@ void SGroomEditorViewport::Construct(const FArguments& InArgs)
 	FGroomViewportLODCommands::Register();
 
 	SEditorViewport::Construct( SEditorViewport::FArguments() );
+	
+	UE::AdvancedPreviewScene::BindDefaultOnSettingsChangedHandler(AdvancedPreviewScene, Client);
 }
 
 SGroomEditorViewport::~SGroomEditorViewport()
diff --git a/Engine/Plugins/TextureGraph/Source/TextureGraphEditor/Private/STG_EditorViewport.cpp b/Engine/Plugins/TextureGraph/Source/TextureGraphEditor/Private/STG_EditorViewport.cpp
--- a/Engine/Plugins/TextureGraph/Source/TextureGraphEditor/Private/STG_EditorViewport.cpp
+++ b/Engine/Plugins/TextureGraph/Source/TextureGraphEditor/Private/STG_EditorViewport.cpp
@@ -2,6 +2,7 @@
 
 #include "STG_EditorViewport.h"
 #include "AdvancedPreviewScene.h"
+#include "AdvancedPreviewSceneMenus.h"
 #include "AssetViewerSettings.h"
 #include "Brushes/SlateBorderBrush.h"
 #include "Brushes/SlateImageBrush.h"
@@ -260,6 +261,8 @@ void STG_EditorViewport::Construct(const FArguments& InArgs)
 
 	FCoreUObjectDelegates::OnObjectPropertyChanged.AddRaw(this, &STG_EditorViewport::OnObjectPropertyChanged);
 	
+	UE::AdvancedPreviewScene::BindDefaultOnSettingsChangedHandler(PreviewScene, EditorViewportClient);
+	
 	GenerateRendermodeToolbar();
 }
 
@@ -318,7 +321,6 @@ TSharedRef<SEditorViewport> STG_EditorViewport::GetViewportWidget()
 TSharedRef<FEditorViewportClient> STG_EditorViewport::MakeEditorViewportClient()
 {
 	EditorViewportClient = MakeShareable(new FTG_EditorViewportClient(TextureGraphPtr, *PreviewScene.Get(), SharedThis(this)));
-	UAssetViewerSettings::Get()->OnAssetViewerSettingsChanged().AddRaw(this, &STG_EditorViewport::OnAssetViewerSettingsChanged);
 	EditorViewportClient->SetViewLocation(FVector::ZeroVector);
 	EditorViewportClient->SetViewRotation(FRotator(-25.0f, -135.0f, 0.0f));
 	EditorViewportClient->SetViewLocationForOrbiting(FVector::ZeroVector, 500);
@@ -576,21 +578,6 @@ void STG_EditorViewport::GenerateRendermodeToolbar()
 	RenderModeToolBar->Init();
 }
 
-void STG_EditorViewport::OnAssetViewerSettingsChanged(const FName& InPropertyName)
-{
-	/*if (InPropertyName == GET_MEMBER_NAME_CHECKED(FPreviewSceneProfile, bRotateLightingRig) || InPropertyName == NAME_None)
-	{
-		UAssetViewerSettings* Settings = UAssetViewerSettings::Get();
-		const int32 ProfileIndex = AdvancedPreviewScene->GetCurrentProfileIndex();
-		if (Settings->Profiles.IsValidIndex(ProfileIndex) &&
-			Settings->Profiles[ProfileIndex].bRotateLightingRig
-			&& !EditorViewportClient->IsRealtime())
-		{
-			EditorViewportClient->SetRealtime(true);
-		}
-	}*/
-}
-
 TSharedPtr<FExtender> STG_EditorViewport::GetExtenders() const
 {
 	TSharedPtr<FExtender> Result(MakeShareable(new FExtender));
diff --git a/Engine/Plugins/TextureGraph/Source/TextureGraphEditor/Private/STG_EditorViewport.h b/Engine/Plugins/TextureGraph/Source/TextureGraphEditor/Private/STG_EditorViewport.h
--- a/Engine/Plugins/TextureGraph/Source/TextureGraphEditor/Private/STG_EditorViewport.h
+++ b/Engine/Plugins/TextureGraph/Source/TextureGraphEditor/Private/STG_EditorViewport.h
@@ -84,9 +84,6 @@ public:
 
 	void									InitPreviewMesh();
 
-	/** Call back for when the user changes preview scene settings in the UI */
-	void									OnAssetViewerSettingsChanged(const FName& InPropertyName);
-
 	/** @return The editor viewport client */
 	class FEditorViewportClient&			GetViewportClient();
 	void									OnObjectPropertyChanged(UObject* ObjectBeingModified, FPropertyChangedEvent& PropertyChangedEvent);
diff --git a/Engine/Source/Editor/AdvancedPreviewScene/AdvancedPreviewScene.Build.cs b/Engine/Source/Editor/AdvancedPreviewScene/AdvancedPreviewScene.Build.cs
--- a/Engine/Source/Editor/AdvancedPreviewScene/AdvancedPreviewScene.Build.cs
+++ b/Engine/Source/Editor/AdvancedPreviewScene/AdvancedPreviewScene.Build.cs
@@ -14,6 +14,7 @@ public class AdvancedPreviewScene : ModuleRules
 
         PrivateDependencyModuleNames.AddRange(
             new string[] {
+				"CommonMenuExtensions",
                 "Core",
                 "CoreUObject",
                 "Engine",
diff --git a/Engine/Source/Editor/AdvancedPreviewScene/Private/AdvancedPreviewScene.cpp b/Engine/Source/Editor/AdvancedPreviewScene/Private/AdvancedPreviewScene.cpp
--- a/Engine/Source/Editor/AdvancedPreviewScene/Private/AdvancedPreviewScene.cpp
+++ b/Engine/Source/Editor/AdvancedPreviewScene/Private/AdvancedPreviewScene.cpp
@@ -24,6 +24,7 @@
 #include "Framework/Commands/UICommandList.h"
 #include "Framework/Application/SlateApplication.h"
 #include "Editor.h"
+#include "ShowFlagMenuCommands.h"
 
 FAdvancedPreviewScene::FAdvancedPreviewScene(ConstructionValues CVS, float InFloorOffset)
 	: FPreviewScene(CVS)
@@ -232,6 +233,15 @@ void FAdvancedPreviewScene::SetProfileIndex(const int32 InProfileIndex)
 	DefaultSettings->OnAssetViewerSettingsChanged().Broadcast(NAME_None);
 }
 
+FPreviewSceneProfile* FAdvancedPreviewScene::GetCurrentProfile() const
+{
+	if (DefaultSettings->Profiles.IsValidIndex(CurrentProfileIndex))
+	{
+		return &DefaultSettings->Profiles[CurrentProfileIndex];
+	}
+	return nullptr;
+}
+
 void FAdvancedPreviewScene::Tick(float DeltaTime)
 {
 	checkf(DefaultSettings && DefaultSettings->Profiles.IsValidIndex(CurrentProfileIndex), TEXT("Invalid default settings pointer or current profile index"));
@@ -428,9 +438,25 @@ void FAdvancedPreviewScene::BindCommands()
 		FCanExecuteAction(),
 		FIsActionChecked::CreateRaw(this, &FAdvancedPreviewScene::IsEnvironmentEnabled)
 	);
+	
+	UICommandList->MapAction(Commands.ToggleGrid,
+		FExecuteAction::CreateRaw(this, &FAdvancedPreviewScene::HandleToggleGrid),
+		FCanExecuteAction(),
+		FIsActionChecked::CreateRaw(this, &FAdvancedPreviewScene::IsGridEnabled)
+	);
+	
+	// Override the grid toggle command so that toggling the grid toggles the setting first
+	UICommandList->MapAction(FShowFlagMenuCommands::Get().FindCommandForFlag(FEngineShowFlags::SF_Grid),
+		FExecuteAction::CreateRaw(this, &FAdvancedPreviewScene::HandleToggleGrid),
+		FCanExecuteAction(),
+		FIsActionChecked::CreateRaw(this, &FAdvancedPreviewScene::IsGridEnabled)
+	);
 
 	UICommandList->MapAction(Commands.TogglePostProcessing,
-		FExecuteAction::CreateRaw(this, &FAdvancedPreviewScene::HandleTogglePostProcessing));
+		FExecuteAction::CreateRaw(this, &FAdvancedPreviewScene::HandleTogglePostProcessing),
+		FCanExecuteAction(),
+		FIsActionChecked::CreateRaw(this, &FAdvancedPreviewScene::IsPostProcessingEnabled)
+	);
 }
 
 TSharedRef<FUICommandList> FAdvancedPreviewScene::CreateCommandList()
@@ -458,6 +484,20 @@ bool FAdvancedPreviewScene::IsFloorEnabled() const
 	return DefaultSettings->Profiles[CurrentProfileIndex].bShowFloor;
 }
 
+void FAdvancedPreviewScene::HandleToggleGrid()
+{
+	FProperty* GridProperty = FindFProperty<FProperty>(FPreviewSceneProfile::StaticStruct(), GET_MEMBER_NAME_CHECKED(FPreviewSceneProfile, bShowGrid));
+	DefaultSettings->Profiles[CurrentProfileIndex].bShowGrid = !DefaultSettings->Profiles[CurrentProfileIndex].bShowGrid;
+
+	FPropertyChangedEvent PropertyEvent(GridProperty);
+	DefaultSettings->PostEditChangeProperty(PropertyEvent);
+}
+
+bool FAdvancedPreviewScene::IsGridEnabled() const
+{
+	return DefaultSettings->Profiles[CurrentProfileIndex].bShowGrid;
+}
+
 void FAdvancedPreviewScene::HandleTogglePostProcessing()
 {
 	FPreviewSceneProfile& Profile = DefaultSettings->Profiles[CurrentProfileIndex];
@@ -471,6 +511,11 @@ void FAdvancedPreviewScene::HandleTogglePostProcessing()
 	DefaultSettings->PostEditChangeProperty(PropertyEvent);	
 }
 
+bool FAdvancedPreviewScene::IsPostProcessingEnabled() const
+{
+	return DefaultSettings->Profiles[CurrentProfileIndex].bPostProcessingEnabled;
+}
+
 void FAdvancedPreviewScene::OnAssetViewerSettingsRefresh(const FName& InPropertyName)
 {
 	// If the profile was changed, update the current index and the scene.
@@ -482,6 +527,8 @@ void FAdvancedPreviewScene::OnAssetViewerSettingsRefresh(const FName& InProperty
 		UILightingRigRotationDelta = 0.0f;
 
 		UpdateScene(DefaultSettings->Profiles[CurrentProfileIndex]);
+		
+		OnProfileChangedDelegate.Broadcast(DefaultSettings->Profiles[CurrentProfileIndex], InPropertyName);
 	}
 	else if (DefaultSettings->Profiles.IsValidIndex(CurrentProfileIndex))
 	{
@@ -496,5 +543,7 @@ void FAdvancedPreviewScene::OnAssetViewerSettingsRefresh(const FName& InProperty
 		PreviousRotation = DefaultSettings->Profiles[CurrentProfileIndex].LightingRigRotation;
 
 		UpdateScene(DefaultSettings->Profiles[CurrentProfileIndex], bUpdateSkyLight || bNameNone, bUpdateEnvironment || bNameNone, bUpdatePostProcessing || bNameNone, bUpdateDirectionalLight || bNameNone);
+		
+		OnProfileChangedDelegate.Broadcast(DefaultSettings->Profiles[CurrentProfileIndex], InPropertyName);
 	}
 }
diff --git a/Engine/Source/Editor/AdvancedPreviewScene/Private/AdvancedPreviewSceneCommands.cpp b/Engine/Source/Editor/AdvancedPreviewScene/Private/AdvancedPreviewSceneCommands.cpp
--- a/Engine/Source/Editor/AdvancedPreviewScene/Private/AdvancedPreviewSceneCommands.cpp
+++ b/Engine/Source/Editor/AdvancedPreviewScene/Private/AdvancedPreviewSceneCommands.cpp
@@ -12,6 +12,7 @@ void FAdvancedPreviewSceneCommands::RegisterCommands()
 {
 	UI_COMMAND(ToggleEnvironment, "Toggle Environment", "Toggles Environment visibility", EUserInterfaceActionType::ToggleButton, FInputChord(EKeys::I));
 	UI_COMMAND(ToggleFloor, "Toggle Floor", "Toggles floor visibility", EUserInterfaceActionType::ToggleButton, FInputChord(EKeys::O));
+	UI_COMMAND(ToggleGrid, "Toggle Grid", "Toggles grid visibility", EUserInterfaceActionType::ToggleButton, FInputChord());
 	UI_COMMAND(TogglePostProcessing, "Toggle Post Processing", "Toggles whether Post Processing is enabled", EUserInterfaceActionType::ToggleButton, FInputChord(EKeys::P));
 }
 
diff --git a/Engine/Source/Editor/AdvancedPreviewScene/Private/AdvancedPreviewSceneMenus.cpp b/Engine/Source/Editor/AdvancedPreviewScene/Private/AdvancedPreviewSceneMenus.cpp
--- a/Engine/Source/Editor/AdvancedPreviewScene/Private/AdvancedPreviewSceneMenus.cpp
+++ b/Engine/Source/Editor/AdvancedPreviewScene/Private/AdvancedPreviewSceneMenus.cpp
@@ -2,7 +2,10 @@
 
 #include "AdvancedPreviewSceneMenus.h"
 
+#include "AdvancedPreviewScene.h"
 #include "AdvancedPreviewSceneCommands.h"
+#include "AssetViewerSettings.h"
+#include "EditorViewportClient.h"
 #include "ToolMenus.h"
 
 #define LOCTEXT_NAMESPACE "AdvancedPreviewSceneMenus"
@@ -10,7 +13,7 @@
 namespace UE::AdvancedPreviewScene::Menus
 {
 
-void ExtendAdvancedPreviewSceneSettings(FName InAssetViewerProfileSubmenuName)
+void ExtendAdvancedPreviewSceneSettings(FName InAssetViewerProfileSubmenuName, const FSettingsOptions& InSettingsOptions)
 {
 	UToolMenu* const Submenu = UToolMenus::Get()->ExtendMenu(InAssetViewerProfileSubmenuName);
 	if (!Submenu)
@@ -23,16 +26,72 @@ void ExtendAdvancedPreviewSceneSettings(FName InAssetViewerProfileSubmenuName)
 		LOCTEXT("PreviewProfileSceneSettingsSectionLabel", "Preview Scene Options")
 	);
 
-	Section.AddMenuEntry(
-		FAdvancedPreviewSceneCommands::Get().ToggleEnvironment,
-		LOCTEXT("ToggleEnvironmentLabel", "Background"),
-		LOCTEXT("ToggleEnvironmentTooltip", "Set the visibility of the preview scene's background.")
-	);
-	Section.AddMenuEntry(
-		FAdvancedPreviewSceneCommands::Get().ToggleFloor,
-		LOCTEXT("ToggleFloorLabel", "Floor"),
-		LOCTEXT("ToggleFloorTooltip", "Set the visibility of the preview scene's floor.")
-	);
+	if (InSettingsOptions.bShowToggleEnvironment)
+	{
+		Section.AddMenuEntry(
+			FAdvancedPreviewSceneCommands::Get().ToggleEnvironment,
+			LOCTEXT("ToggleEnvironmentLabel", "Background"),
+			LOCTEXT("ToggleEnvironmentTooltip", "Set the visibility of the preview scene's background.")
+		);
+	}
+	
+	if (InSettingsOptions.bShowToggleGrid)
+	{
+		Section.AddMenuEntry(
+        	FAdvancedPreviewSceneCommands::Get().ToggleGrid,
+        	LOCTEXT("ToggleGridLabel", "Grid"),
+        	LOCTEXT("ToggleGridTooltip", "Set the visibility of the preview scene's grid.")
+        );	
+	}
+	
+	if (InSettingsOptions.bShowToggleFloor)
+	{
+		Section.AddMenuEntry(
+    		FAdvancedPreviewSceneCommands::Get().ToggleFloor,
+    		LOCTEXT("ToggleFloorLabel", "Floor"),
+    		LOCTEXT("ToggleFloorTooltip", "Set the visibility of the preview scene's floor.")
+    	);
+	}
+	
+	if (InSettingsOptions.bShowTogglePostProcessing)
+	{
+		Section.AddMenuEntry(
+        	FAdvancedPreviewSceneCommands::Get().TogglePostProcessing,
+        	LOCTEXT("TogglePostProcessingLabel", "Post Processing"),
+        	LOCTEXT("TogglePostProcessingTooltip", "Set whether the preview scene includes post processing.")
+        );
+	}
+}
+
+}
+
+namespace UE::AdvancedPreviewScene
+{
+
+void BindDefaultOnSettingsChangedHandler(const TSharedPtr<FAdvancedPreviewScene>& Scene, const TSharedPtr<FEditorViewportClient>& Client)
+{
+	Scene->OnProfileChanged().AddStatic(&DefaultOnSettingsChangedHandler, Client.ToWeakPtr());
+	
+	if (FPreviewSceneProfile* Profile = Scene->GetCurrentProfile())
+	{
+		// Do the initial sync of profile settings to client settings
+		DefaultOnSettingsChangedHandler(*Profile, NAME_None, Client);
+	}
+}
+
+void DefaultOnSettingsChangedHandler(const FPreviewSceneProfile& Profile, FName PropertyName, TWeakPtr<FEditorViewportClient> Client)
+{
+	if (TSharedPtr<FEditorViewportClient> ViewportClient = Client.Pin())
+	{
+		Profile.SetShowFlags(ViewportClient->EngineShowFlags);
+		
+		if (Profile.bRotateLightingRig && !ViewportClient->IsRealtime())
+		{
+			ViewportClient->SetRealtime(true);
+		}
+		
+		ViewportClient->Invalidate();
+	}
 }
 
 }
diff --git a/Engine/Source/Editor/AdvancedPreviewScene/Public/AdvancedPreviewScene.h b/Engine/Source/Editor/AdvancedPreviewScene/Public/AdvancedPreviewScene.h
--- a/Engine/Source/Editor/AdvancedPreviewScene/Public/AdvancedPreviewScene.h
+++ b/Engine/Source/Editor/AdvancedPreviewScene/Public/AdvancedPreviewScene.h
@@ -7,6 +7,7 @@
 #pragma once
 
 #include "CoreMinimal.h"
+#include "AssetViewerSettings.h"
 #include "UnrealClient.h"
 #include "Stats/Stats.h"
 #include "InputCoreTypes.h"
@@ -26,6 +27,9 @@ class FUICommandList;
 class ADVANCEDPREVIEWSCENE_API FAdvancedPreviewScene : public FPreviewScene, public FTickableEditorObject
 {
 public:
+	/** An event for when a profile changes. Includes the changed profile and (optionally) the name of the property in the profile that changed. */
+	DECLARE_MULTICAST_DELEGATE_TwoParams(FOnProfileChanged, const FPreviewSceneProfile&, FName);
+
 	FAdvancedPreviewScene(ConstructionValues CVS, float InFloorOffset = 0.0f);
 	~FAdvancedPreviewScene();
 
@@ -40,6 +44,8 @@ public:
 	virtual ETickableTickType GetTickableTickType() const override { return ETickableTickType::Always; }
 	virtual TStatId GetStatId() const override;
 	/* End FTickableEditorObject */
+	
+	FOnProfileChanged& OnProfileChanged() { return OnProfileChangedDelegate; }
 
 	const bool HandleViewportInput(FViewport* InViewport, FInputDeviceId DeviceId, FKey Key, float Delta, float DeltaTime, int32 NumSamples, bool bGamepad);
 	const bool HandleInputKey(const FInputKeyEventArgs& EventArgs);
@@ -52,6 +58,7 @@ public:
 	float GetFloorOffset() const;
 	void SetFloorOffset(const float InFloorOffset);
 	void SetProfileIndex(const int32 InProfileIndex);
+	FPreviewSceneProfile* GetCurrentProfile() const;
 
 	const UStaticMeshComponent* GetFloorMeshComponent() const;
 	const float GetSkyRotation() const;
@@ -60,13 +67,6 @@ public:
 
 	TSharedPtr<FUICommandList> GetCommandList() const { return UICommandList; }
 
-protected:
-	/** Create and map the command list. */
-	virtual void BindCommands();
-	
-	/** Add commands to a provided command list. */
-	virtual TSharedRef<FUICommandList> CreateCommandList();
-
 	/** Toggle the sky sphere on and off */
 	void HandleToggleEnvironment();
 	
@@ -79,13 +79,31 @@ protected:
 	/** Whether the floor is on */
 	bool IsFloorEnabled() const;
 
+	/** Toggle the grid on and off. */
+	void HandleToggleGrid();
+	
+	/** Whether the grid is enabled. */
+	bool IsGridEnabled() const;
+
 	/** Toggle post processing on and off */
 	void HandleTogglePostProcessing();
+	
+	/** Whether post processing is enabled */
+	bool IsPostProcessingEnabled() const;
 
+protected:
+	/** Create and map the command list. */
+	virtual void BindCommands();
+	
+	/** Add commands to a provided command list. */
+	virtual TSharedRef<FUICommandList> CreateCommandList();
+	
 	/** Handle refreshing the scene when settings change */
 	void OnAssetViewerSettingsRefresh(const FName& InPropertyName);
 
 protected:
+	FOnProfileChanged OnProfileChangedDelegate;
+
 	UStaticMeshComponent* SkyComponent;
 	UMaterialInstanceConstant* InstancedSkyMaterial;
 	UPostProcessComponent* PostProcessComponent;
diff --git a/Engine/Source/Editor/AdvancedPreviewScene/Public/AdvancedPreviewSceneCommands.h b/Engine/Source/Editor/AdvancedPreviewScene/Public/AdvancedPreviewSceneCommands.h
--- a/Engine/Source/Editor/AdvancedPreviewScene/Public/AdvancedPreviewSceneCommands.h
+++ b/Engine/Source/Editor/AdvancedPreviewScene/Public/AdvancedPreviewSceneCommands.h
@@ -31,6 +31,9 @@ public:
 	/** Toggles floor visibility */
 	TSharedPtr< FUICommandInfo > ToggleFloor;
 
+	/** Toggles the grid */
+	TSharedPtr< FUICommandInfo > ToggleGrid;
+	
 	/** Toggles post processing */
 	TSharedPtr< FUICommandInfo > TogglePostProcessing;
 
diff --git a/Engine/Source/Editor/AdvancedPreviewScene/Public/AdvancedPreviewSceneMenus.h b/Engine/Source/Editor/AdvancedPreviewScene/Public/AdvancedPreviewSceneMenus.h
--- a/Engine/Source/Editor/AdvancedPreviewScene/Public/AdvancedPreviewSceneMenus.h
+++ b/Engine/Source/Editor/AdvancedPreviewScene/Public/AdvancedPreviewSceneMenus.h
@@ -3,16 +3,64 @@
 #pragma once
 
 #include "UObject/NameTypes.h"
+#include "Templates/SharedPointer.h"
+
+class FUICommandList;
+class FAdvancedPreviewScene;
+class FEditorViewportClient;
+struct FPreviewSceneProfile;
 
 namespace UE::AdvancedPreviewScene::Menus
 {
 
-ADVANCEDPREVIEWSCENE_API void ExtendAdvancedPreviewSceneSettings(FName InAssetViewerProfileSubmenuName);
+struct ADVANCEDPREVIEWSCENE_API FSettingsOptions
+{
+	bool bShowToggleEnvironment = true;
+	bool bShowToggleGrid = true;
+	bool bShowToggleFloor = true;
+	bool bShowTogglePostProcessing = true;
+
+	FSettingsOptions& ShowToggleEnvironment(bool bInShowToggleEnvironment)
+	{
+		bShowToggleEnvironment = bInShowToggleEnvironment;
+		return *this;
+	}
+
+	FSettingsOptions& ShowToggleGrid(bool bInShowToggleGrid)
+	{
+		bShowToggleGrid = bInShowToggleGrid;
+		return *this;
+	}
+
+	FSettingsOptions& ShowToggleFloor(bool bInShowToggleFloor)
+	{
+		bShowToggleFloor = bInShowToggleFloor;
+		return *this;
+	}
+
+	FSettingsOptions& ShowTogglePostProcessing(bool bInShowTogglePostProcessing)
+	{
+		bShowTogglePostProcessing = bInShowTogglePostProcessing;
+		return *this;
+	}
+};
+
+ADVANCEDPREVIEWSCENE_API void ExtendAdvancedPreviewSceneSettings(FName InAssetViewerProfileSubmenuName, const FSettingsOptions& InSettingsOptions = {});
 
 }
 
-class AdvancedPreviewSceneMenus
+namespace UE::AdvancedPreviewScene
 {
-public:
-	
-};
+
+/**
+ * A default handler for `FAdvancedPreviewScene::OnProfileChanged()`. Syncs the profile to the provided client's show flags and other settings.
+ */
+ADVANCEDPREVIEWSCENE_API void DefaultOnSettingsChangedHandler(const FPreviewSceneProfile& Profile, FName PropertyName, TWeakPtr<FEditorViewportClient> Client);
+
+/**
+ * Attaches `DefaultOnSettingsChangedHandler()` to the provided scene's `OnProfileChanged` delegate & calls the function with the current profile to syncronize the client.
+ * This is a convenience that alleviates the annoyance of having to manually downcast subtypes of FEditorViewporClient.
+ */
+ADVANCEDPREVIEWSCENE_API void BindDefaultOnSettingsChangedHandler(const TSharedPtr<FAdvancedPreviewScene>& Scene, const TSharedPtr<FEditorViewportClient>& Client);
+
+}
diff --git a/Engine/Source/Editor/CommonMenuExtensions/Private/ShowFlagMenuCommands.cpp b/Engine/Source/Editor/CommonMenuExtensions/Private/ShowFlagMenuCommands.cpp
--- a/Engine/Source/Editor/CommonMenuExtensions/Private/ShowFlagMenuCommands.cpp
+++ b/Engine/Source/Editor/CommonMenuExtensions/Private/ShowFlagMenuCommands.cpp
@@ -211,6 +211,18 @@ void FShowFlagMenuCommands::PopulateAllShowFlagsSection(FToolMenuSection& Sectio
 	);
 }
 
+TSharedPtr<const FUICommandInfo> FShowFlagMenuCommands::FindCommandForFlag(FEngineShowFlags::EShowFlag FlagIndex) const
+{
+	for (int32 Index = 0; Index < ShowFlagCommands.Num(); ++Index)
+	{
+		if (ShowFlagCommands[Index].FlagIndex == FlagIndex)
+		{
+			return ShowFlagCommands[Index].ShowMenuItem;
+		}
+	}
+	return nullptr;
+}
+
 void FShowFlagMenuCommands::CreateSubMenuIfRequired(FToolMenuSection& Section, const FShowFlagFilter& Filter, EShowFlagGroup Group, const FName SubMenuName, const FText& MenuLabel, const FText& ToolTip, const FName IconName) const
 {
 	const FShowFlagFilter::FGroupedShowFlagIndices& GroupedFlagIndices = Filter.GetFilteredIndices();
diff --git a/Engine/Source/Editor/CommonMenuExtensions/Public/ShowFlagMenuCommands.h b/Engine/Source/Editor/CommonMenuExtensions/Public/ShowFlagMenuCommands.h
--- a/Engine/Source/Editor/CommonMenuExtensions/Public/ShowFlagMenuCommands.h
+++ b/Engine/Source/Editor/CommonMenuExtensions/Public/ShowFlagMenuCommands.h
@@ -59,6 +59,8 @@ public:
 	void PopulateAllShowFlagsSection(
 		FToolMenuSection& Section, const FShowFlagFilter& Filter = FShowFlagFilter(FShowFlagFilter::IncludeAllFlagsByDefault)
 	) const;
+	
+	TSharedPtr<const FUICommandInfo> FindCommandForFlag(FEngineShowFlags::EShowFlag FlagIndex) const;
 
 private:
 	static void StaticCreateShowFlagsSubMenu(UToolMenu* Menu, TArray<uint32> FlagIndices, int32 EntryOffset);
diff --git a/Engine/Source/Editor/MaterialEditor/Private/SMaterialEditorViewport.cpp b/Engine/Source/Editor/MaterialEditor/Private/SMaterialEditorViewport.cpp
--- a/Engine/Source/Editor/MaterialEditor/Private/SMaterialEditorViewport.cpp
+++ b/Engine/Source/Editor/MaterialEditor/Private/SMaterialEditorViewport.cpp
@@ -285,22 +285,16 @@ void SMaterialEditor3DPreviewViewport::Construct(const FArguments& InArgs)
 
 	SetPreviewAsset( GUnrealEd->GetThumbnailManager()->EditorSphere );
 
-	UAssetViewerSettings* Settings = UAssetViewerSettings::Get();
-	const int32 ProfileIndex = AdvancedPreviewScene->GetCurrentProfileIndex();
-	if (Settings->Profiles.IsValidIndex(ProfileIndex))
-	{
-		AdvancedPreviewScene->SetEnvironmentVisibility(Settings->Profiles[ProfileIndex].bShowEnvironment, true);
-	}
-
 	OnPropertyChangedHandle = FCoreUObjectDelegates::FOnObjectPropertyChanged::FDelegate::CreateRaw(this, &SMaterialEditor3DPreviewViewport::OnPropertyChanged);
 	OnPropertyChangedHandleDelegateHandle = FCoreUObjectDelegates::OnObjectPropertyChanged.Add(OnPropertyChangedHandle);
+	
+	UE::AdvancedPreviewScene::BindDefaultOnSettingsChangedHandler(AdvancedPreviewScene, EditorViewportClient);
 }
 
 SMaterialEditor3DPreviewViewport::~SMaterialEditor3DPreviewViewport()
 {
 	CastChecked<UEditorEngine>(GEngine)->OnPreviewFeatureLevelChanged().Remove(PreviewFeatureLevelChangedHandle);
 	
-	UAssetViewerSettings::Get()->OnAssetViewerSettingsChanged().RemoveAll(this);
 	if (PreviewMeshComponent != nullptr)
 	{
 		PreviewMeshComponent->OverrideMaterials.Empty();
@@ -898,27 +892,6 @@ bool SMaterialEditor3DPreviewViewport::IsTogglePreviewBackgroundChecked() const
 	return false;
 }
 
-
-void SMaterialEditor3DPreviewViewport::OnAssetViewerSettingsChanged(const FName& InPropertyName)
-{
-	if (InPropertyName == GET_MEMBER_NAME_CHECKED(FPreviewSceneProfile, bRotateLightingRig) || InPropertyName == NAME_None)
-	{
-		UAssetViewerSettings* Settings = UAssetViewerSettings::Get();
-		const int32 ProfileIndex = AdvancedPreviewScene->GetCurrentProfileIndex();
-		if (Settings->Profiles.IsValidIndex(ProfileIndex) &&
-			Settings->Profiles[ProfileIndex].bRotateLightingRig
-			&& !EditorViewportClient->IsRealtime())
-		{
-			EditorViewportClient->SetRealtime(true);
-		}
-	}
-	else
-	{
-		// If the profile changes, refresh the viewport for it to take effect
-		RefreshViewport();
-	}
-}
-
 TSharedRef<class SEditorViewport> SMaterialEditor3DPreviewViewport::GetViewportWidget()
 {
 	return SharedThis(this);
@@ -937,7 +910,6 @@ void SMaterialEditor3DPreviewViewport::OnFloatingButtonClicked()
 TSharedRef<FEditorViewportClient> SMaterialEditor3DPreviewViewport::MakeEditorViewportClient() 
 {
 	EditorViewportClient = MakeShareable( new FMaterialEditorViewportClient(MaterialEditorPtr, *AdvancedPreviewScene.Get(), SharedThis(this)) );
-	UAssetViewerSettings::Get()->OnAssetViewerSettingsChanged().AddRaw(this, &SMaterialEditor3DPreviewViewport::OnAssetViewerSettingsChanged);
 	EditorViewportClient->SetViewLocation( FVector::ZeroVector );
 	EditorViewportClient->SetViewRotation( FRotator(-15.0f, -90.0f, 0.0f) );
 	EditorViewportClient->SetViewLocationForOrbiting( FVector::ZeroVector );
@@ -985,8 +957,11 @@ TSharedPtr<SWidget> SMaterialEditor3DPreviewViewport::BuildViewportToolbar()
 			{
 				const FName AssetViewerProfileMenuName = "MaterialEditor.ViewportToolbar.AssetViewerProfile";
 				RightSection.AddEntry(UE::UnrealEd::CreateAssetViewerProfileSubmenu());
-				UE::AdvancedPreviewScene::Menus::ExtendAdvancedPreviewSceneSettings(AssetViewerProfileMenuName);
 				UE::MaterialEditor::ExtendPreviewSceneSettingsSubmenu(AssetViewerProfileMenuName);
+				UE::AdvancedPreviewScene::Menus::ExtendAdvancedPreviewSceneSettings(
+					AssetViewerProfileMenuName,
+					UE::AdvancedPreviewScene::Menus::FSettingsOptions().ShowToggleGrid(false)
+				);
 				UE::UnrealEd::ExtendPreviewSceneSettingsWithTabEntry(AssetViewerProfileMenuName);
 			}
 		}
@@ -994,9 +969,9 @@ TSharedPtr<SWidget> SMaterialEditor3DPreviewViewport::BuildViewportToolbar()
 
 	FToolMenuContext ViewportToolbarContext;
 	{
-		ViewportToolbarContext.AppendCommandList(GetCommandList());
 		ViewportToolbarContext.AppendCommandList(AdvancedPreviewScene->GetCommandList());
-
+		ViewportToolbarContext.AppendCommandList(GetCommandList());
+		
 		// Add the UnrealEd viewport toolbar context.
 		{
 			UUnrealEdViewportToolbarContext* const ContextObject =
diff --git a/Engine/Source/Editor/MaterialEditor/Private/SMaterialEditorViewport.h b/Engine/Source/Editor/MaterialEditor/Private/SMaterialEditorViewport.h
--- a/Engine/Source/Editor/MaterialEditor/Private/SMaterialEditorViewport.h
+++ b/Engine/Source/Editor/MaterialEditor/Private/SMaterialEditorViewport.h
@@ -83,8 +83,6 @@ public:
 	bool IsPreviewMeshFromSelectionChecked() const;
 	void TogglePreviewBackground();
 	bool IsTogglePreviewBackgroundChecked() const;
-	/** Call back for when the user changes preview scene settings in the UI */
-	void OnAssetViewerSettingsChanged(const FName& InPropertyName);
 
 	// ICommonEditorViewportToolbarInfoProvider interface
 	virtual TSharedRef<class SEditorViewport> GetViewportWidget() override;
diff --git a/Engine/Source/Editor/Persona/Private/AnimationEditorViewportClient.cpp b/Engine/Source/Editor/Persona/Private/AnimationEditorViewportClient.cpp
--- a/Engine/Source/Editor/Persona/Private/AnimationEditorViewportClient.cpp
+++ b/Engine/Source/Editor/Persona/Private/AnimationEditorViewportClient.cpp
@@ -289,8 +289,6 @@ FAnimationViewportClient::~FAnimationViewportClient()
 	}
 	OnPhysicsCreatedDelegateHandle.Reset();
 	OnMeshChangedDelegateHandle.Reset();
-	
-	UAssetViewerSettings::Get()->OnAssetViewerSettingsChanged().RemoveAll(this);
 }
 
 void FAnimationViewportClient::Initialize()
@@ -314,11 +312,6 @@ void FAnimationViewportClient::Initialize()
         }));
 	}
 
-	// Register delegate to update the show flags when the post processing is turned on or off
-	UAssetViewerSettings::Get()->OnAssetViewerSettingsChanged().AddSP(this, &FAnimationViewportClient::OnAssetViewerSettingsChanged);
-	// Set correct flags according to current profile settings
-	UAssetViewerSettings::GetCurrentUserProjectProfile().SetShowFlags(EngineShowFlags);
-
 	// Setup bones to draw on initialise
 	UpdateBonesToDraw();
 }
@@ -2614,31 +2607,6 @@ int32 FAnimationViewportClient::GetShowMeshStats() const
 	return ConfigOption->ShowMeshStats;
 }
 
-void FAnimationViewportClient::OnAssetViewerSettingsChanged(const FName& InPropertyName)
-{
-	if (InPropertyName == GET_MEMBER_NAME_CHECKED(FPreviewSceneProfile, bPostProcessingEnabled) || InPropertyName == NAME_None)
-	{
-		UAssetViewerSettings* Settings = UAssetViewerSettings::Get();
-		const int32 ProfileIndex = GetPreviewScene()->GetCurrentProfileIndex();
-		if (Settings->Profiles.IsValidIndex(ProfileIndex))
-		{			
-			Settings->Profiles[ProfileIndex].SetShowFlags(EngineShowFlags);
-			SetAdvancedShowFlagsForScene(Settings->Profiles[ProfileIndex].bPostProcessingEnabled);
-		}
-	}
-}
-
-void FAnimationViewportClient::SetAdvancedShowFlagsForScene(const bool bAdvancedShowFlags)
-{	
-	if (bAdvancedShowFlags)
-	{
-		EngineShowFlags.EnableAdvancedFeatures();
-	}
-	else
-	{
-		EngineShowFlags.DisableAdvancedFeatures();
-	}
-}
 void FAnimationViewportClient::SetPlaybackSpeedMode(EAnimationPlaybackSpeeds::Type InMode)
 {
 	if (FAnimationEditorPreviewScene* AnimationEditorPreviewScene = GetAnimPreviewScenePtr().Get())
diff --git a/Engine/Source/Editor/Persona/Private/SAnimationEditorViewport.cpp b/Engine/Source/Editor/Persona/Private/SAnimationEditorViewport.cpp
--- a/Engine/Source/Editor/Persona/Private/SAnimationEditorViewport.cpp
+++ b/Engine/Source/Editor/Persona/Private/SAnimationEditorViewport.cpp
@@ -294,13 +294,13 @@ TSharedPtr<SWidget> SAnimationEditorViewport::BuildViewportToolbar()
 
 	FToolMenuContext ViewportToolbarContext;
 	{
-		ViewportToolbarContext.AppendCommandList(GetCommandList());
-		
 		if (TSharedPtr<IPersonaPreviewScene> PreviewScene = PreviewScenePtr.Pin())
 		{
 			ViewportToolbarContext.AppendCommandList(PreviewScene->GetCommandList());
 		}
-	
+
+		ViewportToolbarContext.AppendCommandList(GetCommandList());
+		
 		if (TSharedPtr<SAnimationEditorViewportTabBody> TabBody = TabBodyPtr.Pin())
 		{
 			ViewportToolbarContext.AppendCommandList(TabBody->GetCommandList());
@@ -435,8 +435,6 @@ SAnimationEditorViewportTabBody::~SAnimationEditorViewportTabBody()
 	LevelViewportClient.Reset();
 
 	PendingTransaction.Reset();
-	
-	UAssetViewerSettings::Get()->OnAssetViewerSettingsChanged().RemoveAll(this);
 }
 
 bool SAnimationEditorViewportTabBody::CanUseGizmos() const
@@ -907,7 +905,10 @@ void SAnimationEditorViewportTabBody::Construct(const FArguments& InArgs, const
 
 	AddSkinWeightProfileNotification();
 	
-	UAssetViewerSettings::Get()->OnAssetViewerSettingsChanged().AddSP(this, &SAnimationEditorViewportTabBody::OnAssetViewerSettingsChanged);
+	if (TSharedPtr<FAnimationEditorPreviewScene> PreviewScene = PreviewScenePtr.Pin())
+	{
+		UE::AdvancedPreviewScene::BindDefaultOnSettingsChangedHandler(PreviewScene, LevelViewportClient);
+	}
 }
 
 void SAnimationEditorViewportTabBody::BindCommands()
@@ -1481,11 +1482,6 @@ int32 SAnimationEditorViewportTabBody::GetLODModelCount() const
 	return 0;
 }
 
-void SAnimationEditorViewportTabBody::OnAssetViewerSettingsChanged(const FName& InPropertyName)
-{
-	RefreshViewport();
-}
-
 void SAnimationEditorViewportTabBody::OnShowMorphTargets()
 {
 	ForEachDebugMesh([](UDebugSkelMeshComponent* InMesh)
diff --git a/Engine/Source/Editor/Persona/Private/SAnimationEditorViewport.h b/Engine/Source/Editor/Persona/Private/SAnimationEditorViewport.h
--- a/Engine/Source/Editor/Persona/Private/SAnimationEditorViewport.h
+++ b/Engine/Source/Editor/Persona/Private/SAnimationEditorViewport.h
@@ -548,7 +548,6 @@ private:
 	void OnSetTurnTableSpeed(int32 SpeedIndex);
 	void OnSetTurnTableMode(int32 ModeIndex);
 	bool IsTurnTableModeSelected(int32 ModeIndex) const;
-	void OnAssetViewerSettingsChanged(const FName& InPropertyName);
 
 	FPersonaViewportKeyDownDelegate OnKeyDownDelegate;
 
diff --git a/Engine/Source/Editor/Persona/Public/AnimationEditorViewportClient.h b/Engine/Source/Editor/Persona/Public/AnimationEditorViewportClient.h
--- a/Engine/Source/Editor/Persona/Public/AnimationEditorViewportClient.h
+++ b/Engine/Source/Editor/Persona/Public/AnimationEditorViewportClient.h
@@ -428,12 +428,6 @@ private:
 	/** Handle the view being focused from the preview scene */
 	void HandleFocusViews();
 
-	/** Delegate for preview profile is changed (used for updating show flags) */
-	void OnAssetViewerSettingsChanged(const FName& InPropertyName);
-
-	/** Sets up the ShowFlag according to the current preview scene profile */
-	void SetAdvancedShowFlagsForScene(const bool bAdvancedShowFlags);
-	
 	/** Computes a bounding box for the selected section of the preview mesh component.
 	    If there is no selected section, returns an empty box. */
 	FBox ComputeBoundingBoxForSelectedEditorSection() const;
diff --git a/Engine/Source/Editor/StaticMeshEditor/Private/SStaticMeshEditorViewport.cpp b/Engine/Source/Editor/StaticMeshEditor/Private/SStaticMeshEditorViewport.cpp
--- a/Engine/Source/Editor/StaticMeshEditor/Private/SStaticMeshEditorViewport.cpp
+++ b/Engine/Source/Editor/StaticMeshEditor/Private/SStaticMeshEditorViewport.cpp
@@ -81,7 +81,7 @@ void SStaticMeshEditorViewport::Construct(const FArguments& InArgs)
 
 	FCoreUObjectDelegates::OnObjectPropertyChanged.AddRaw(this, &SStaticMeshEditorViewport::OnObjectPropertyChanged);
 	
-	UAssetViewerSettings::Get()->OnAssetViewerSettingsChanged().AddSP(this, &SStaticMeshEditorViewport::OnAssetViewerSettingsChanged);
+	UE::AdvancedPreviewScene::BindDefaultOnSettingsChangedHandler(PreviewScene, EditorViewportClient);
 }
 
 void SStaticMeshEditorViewport::PopulateViewportOverlays(TSharedRef<SOverlay> Overlay)
@@ -126,8 +126,6 @@ SStaticMeshEditorViewport::~SStaticMeshEditorViewport()
 	{
 		EditorViewportClient->Viewport = NULL;
 	}
-	
-	UAssetViewerSettings::Get()->OnAssetViewerSettingsChanged().RemoveAll(this);
 }
 
 void SStaticMeshEditorViewport::PopulateOverlayText(const TArrayView<FOverlayTextItem> TextItems)
@@ -327,11 +325,6 @@ void SStaticMeshEditorViewport::SetFloorAndEnvironmentVisibility(bool bVisible)
 	}
 }
 
-void SStaticMeshEditorViewport::OnAssetViewerSettingsChanged(const FName& InPropertyName)
-{
-	RefreshViewport();
-}
-
 FMargin SStaticMeshEditorViewport::GetOverlayMargin() const
 {
 	return UE::UnrealEd::ShowOldViewportToolbars()
@@ -838,9 +831,9 @@ void SStaticMeshEditorViewport::BindCommands()
 
 	CommandList->MapAction(
 		Commands.SetShowGrid,
-		FExecuteAction::CreateSP( EditorViewportClientRef, &FStaticMeshEditorViewportClient::SetShowGrid ),
+		FExecuteAction::CreateRaw( PreviewScene.Get(), &FAdvancedPreviewScene::HandleToggleGrid ),
 		FCanExecuteAction(),
-		FIsActionChecked::CreateSP( EditorViewportClientRef, &FStaticMeshEditorViewportClient::IsSetShowGridChecked ) );
+		FIsActionChecked::CreateRaw( PreviewScene.Get(), &FAdvancedPreviewScene::IsGridEnabled ) );
 
 	CommandList->MapAction(
 		Commands.SetShowBounds,
@@ -1039,8 +1032,8 @@ TSharedPtr<SWidget> SStaticMeshEditorViewport::BuildViewportToolbar()
 
 	FToolMenuContext ViewportToolbarContext;
 	{
-		ViewportToolbarContext.AppendCommandList(GetCommandList());
 		ViewportToolbarContext.AppendCommandList(PreviewScene->GetCommandList());
+		ViewportToolbarContext.AppendCommandList(GetCommandList());
 
 		// Add the UnrealEd viewport toolbar context.
 		{
diff --git a/Engine/Source/Editor/StaticMeshEditor/Private/SStaticMeshEditorViewport.h b/Engine/Source/Editor/StaticMeshEditor/Private/SStaticMeshEditorViewport.h
--- a/Engine/Source/Editor/StaticMeshEditor/Private/SStaticMeshEditorViewport.h
+++ b/Engine/Source/Editor/StaticMeshEditor/Private/SStaticMeshEditorViewport.h
@@ -205,8 +205,6 @@ private:
 
 	void SetFloorAndEnvironmentVisibility(bool bVisible);
 
-	void OnAssetViewerSettingsChanged(const FName& InPropertyName);
-
 	/** Returns the appropriate margin for the overlay, based on the current viewport toolbar configuration */
 	FMargin GetOverlayMargin() const;
 private:
diff --git a/Engine/Source/Editor/StaticMeshEditor/Private/StaticMeshEditorViewportClient.cpp b/Engine/Source/Editor/StaticMeshEditor/Private/StaticMeshEditorViewportClient.cpp
--- a/Engine/Source/Editor/StaticMeshEditor/Private/StaticMeshEditorViewportClient.cpp
+++ b/Engine/Source/Editor/StaticMeshEditor/Private/StaticMeshEditorViewportClient.cpp
@@ -95,16 +95,10 @@ FStaticMeshEditorViewportClient::FStaticMeshEditorViewportClient(TWeakPtr<IStati
 	AdvancedPreviewScene = static_cast<FAdvancedPreviewScene*>(PreviewScene);
 
 	SetPreviewMesh(InPreviewStaticMesh, InPreviewStaticMeshComponent);
-
-	// Register delegate to update the show flags when the post processing is turned on or off
-	UAssetViewerSettings::Get()->OnAssetViewerSettingsChanged().AddRaw(this, &FStaticMeshEditorViewportClient::OnAssetViewerSettingsChanged);
-	// Set correct flags according to current profile settings
-	SetAdvancedShowFlagsForScene(UAssetViewerSettings::Get()->Profiles[GetMutableDefault<UEditorPerProjectUserSettings>()->AssetViewerProfileIndex].bPostProcessingEnabled);
 }
 
 FStaticMeshEditorViewportClient::~FStaticMeshEditorViewportClient()
 {
-	UAssetViewerSettings::Get()->OnAssetViewerSettingsChanged().RemoveAll(this);
 }
 
 void FStaticMeshEditorViewportClient::Tick(float DeltaSeconds)
@@ -1424,31 +1418,6 @@ void FStaticMeshEditorViewportClient::PerspectiveCameraMoved()
 	ToggleOrbitCamera(bWasOrbit);
 }
 
-void FStaticMeshEditorViewportClient::OnAssetViewerSettingsChanged(const FName& InPropertyName)
-{
-	if (InPropertyName == GET_MEMBER_NAME_CHECKED(FPreviewSceneProfile, bPostProcessingEnabled) || InPropertyName == NAME_None)
-	{
-		UAssetViewerSettings* Settings = UAssetViewerSettings::Get();
-		const int32 ProfileIndex = AdvancedPreviewScene->GetCurrentProfileIndex();
-		if (Settings->Profiles.IsValidIndex(ProfileIndex))
-		{
-			SetAdvancedShowFlagsForScene(Settings->Profiles[ProfileIndex].bPostProcessingEnabled);
-		}		
-	}
-}
-
-void FStaticMeshEditorViewportClient::SetAdvancedShowFlagsForScene(const bool bAdvancedShowFlags)
-{	
-	if (bAdvancedShowFlags)
-	{
-		EngineShowFlags.EnableAdvancedFeatures();
-	}
-	else
-	{
-		EngineShowFlags.DisableAdvancedFeatures();
-	}
-}
-
 void FStaticMeshEditorViewportClient::UpdateSimpleCollisionDisplay()
 {
 	if (StaticMeshComponent != nullptr)
diff --git a/Engine/Source/Editor/StaticMeshEditor/Private/StaticMeshEditorViewportClient.h b/Engine/Source/Editor/StaticMeshEditor/Private/StaticMeshEditorViewportClient.h
--- a/Engine/Source/Editor/StaticMeshEditor/Private/StaticMeshEditorViewportClient.h
+++ b/Engine/Source/Editor/StaticMeshEditor/Private/StaticMeshEditorViewportClient.h
@@ -173,10 +173,6 @@ protected:
 	// FEditorViewportClient interface
 	virtual void PerspectiveCameraMoved() override;
 
-	/** Call back for when the user changes preview scene settings in the UI */
-	void OnAssetViewerSettingsChanged(const FName& InPropertyName);
-	/** Used to (re)-set the viewport show flags related to post processing*/
-	void SetAdvancedShowFlagsForScene(const bool bAdvancedShowFlags);
 	/** Used to update the display of simple collision **/
 	void UpdateSimpleCollisionDisplay();
 	/** Used to update the display of complex collision **/
