EditorUsability : PreviewScene
- Fixed preview scene not loaded properly on editor startup
- Added options to customize floor, floor rotation, floor scale and floor material
- Added option to clone a profile from existing ones
- Added keyboard shortcuts to navigate to next/previous profile (SHIFT+ALT+Up/Down)
- Reworked the preview scene settings tab UI
- Property AssetPicker can now be customized with additional metadata to set ViewType, ThumbnailSize, ThumbnailScale, SettingsName
- Deprecated "show grid" property since it's not relevant to all context

#Synced-CL 50537297

diff --git a/Engine/Plugins/Experimental/Dataflow/Source/DataflowEditor/Private/Dataflow/DataflowEditorViewportClientBase.cpp b/Engine/Plugins/Experimental/Dataflow/Source/DataflowEditor/Private/Dataflow/DataflowEditorViewportClientBase.cpp
--- a/Engine/Plugins/Experimental/Dataflow/Source/DataflowEditor/Private/Dataflow/DataflowEditorViewportClientBase.cpp
+++ b/Engine/Plugins/Experimental/Dataflow/Source/DataflowEditor/Private/Dataflow/DataflowEditorViewportClientBase.cpp
@@ -86,8 +86,7 @@ void FDataflowEditorViewportClientBase::RegisterDelegates()
 				checkf(DefaultSettings && DefaultSettings->Profiles.IsValidIndex(CurrentProfileIndex), TEXT("Invalid default settings pointer or current profile index"));
 
 				const FPreviewSceneProfile& Profile = DefaultSettings->Profiles[CurrentProfileIndex];
-				EngineShowFlags.Grid = Profile.bShowGrid;
-				DrawHelper.bDrawGrid = Profile.bShowGrid;
+				DrawHelper.bDrawGrid = EngineShowFlags.Grid;
 
 				if (Profile.bPostProcessingEnabled)
 				{
@@ -105,8 +104,7 @@ void FDataflowEditorViewportClientBase::RegisterDelegates()
 
 	OnAssetViewerSettingsChangedDelegateHandle = UAssetViewerSettings::Get()->OnAssetViewerSettingsChanged().AddLambda([this, SyncEngineShowFlags](const FName& InPropertyName)
 	{
-		if ((InPropertyName == GET_MEMBER_NAME_CHECKED(FPreviewSceneProfile, bShowGrid)) ||
-			(InPropertyName == GET_MEMBER_NAME_CHECKED(FPreviewSceneProfile, bPostProcessingEnabled)) ||
+		if ((InPropertyName == GET_MEMBER_NAME_CHECKED(FPreviewSceneProfile, bPostProcessingEnabled)) ||
 			(InPropertyName == NAME_None))		// NAME_None is passed in when the current profile changes
 		{
 			SyncEngineShowFlags();
diff --git a/Engine/Plugins/Experimental/ProceduralVegetationEditor/Source/ProceduralVegetationEditor/Private/SPVEditorViewport.cpp b/Engine/Plugins/Experimental/ProceduralVegetationEditor/Source/ProceduralVegetationEditor/Private/SPVEditorViewport.cpp
--- a/Engine/Plugins/Experimental/ProceduralVegetationEditor/Source/ProceduralVegetationEditor/Private/SPVEditorViewport.cpp
+++ b/Engine/Plugins/Experimental/ProceduralVegetationEditor/Source/ProceduralVegetationEditor/Private/SPVEditorViewport.cpp
@@ -144,7 +144,7 @@ TSharedPtr<SWidget> SPVEditorViewport::BuildViewportToolbar()
 		RightSection.AddEntry(UE::UnrealEd::CreateAssetViewerProfileSubmenu());
 		UE::AdvancedPreviewScene::Menus::ExtendAdvancedPreviewSceneSettings(
 			"PVE.ViewportToolbar.AssetViewerProfile",
-			UE::AdvancedPreviewScene::Menus::FSettingsOptions().ShowToggleGrid(false)
+			UE::AdvancedPreviewScene::Menus::FSettingsOptions()
 		);
 
 		RightSection.AddEntry(CreateVisualizationModeToolbarMenu());
diff --git a/Engine/Plugins/Experimental/ProceduralVegetationEditor/Source/ProceduralVegetationEditor/Private/Visualizations/PVDataVisualization.cpp b/Engine/Plugins/Experimental/ProceduralVegetationEditor/Source/ProceduralVegetationEditor/Private/Visualizations/PVDataVisualization.cpp
--- a/Engine/Plugins/Experimental/ProceduralVegetationEditor/Source/ProceduralVegetationEditor/Private/Visualizations/PVDataVisualization.cpp
+++ b/Engine/Plugins/Experimental/ProceduralVegetationEditor/Source/ProceduralVegetationEditor/Private/Visualizations/PVDataVisualization.cpp
@@ -90,10 +90,7 @@ FPCGSetupSceneFunc FPVDataVisualization::GetViewportSetupFunc(const UPCGSettings
 		// Reapply view mode settings as for some reason SetFloorVisibility changes EngineFlags
 		InOutParams.EditorViewportClient->SetViewMode(InOutParams.EditorViewportClient->GetViewMode());
 		InOutParams.EditorViewportClient->SetRealtime(true);
-		if (!InOutParams.Scene->IsGridEnabled())
-		{
-			InOutParams.Scene->HandleToggleGrid();
-		}
+		
 		if (InOutParams.Scene->IsEnvironmentEnabled())
 		{
 			InOutParams.Scene->HandleToggleEnvironment();
diff --git a/Engine/Plugins/PCG/Source/PCGEditor/Private/Widgets/AssetEditorViewport/SPCGEditorViewport.cpp b/Engine/Plugins/PCG/Source/PCGEditor/Private/Widgets/AssetEditorViewport/SPCGEditorViewport.cpp
--- a/Engine/Plugins/PCG/Source/PCGEditor/Private/Widgets/AssetEditorViewport/SPCGEditorViewport.cpp
+++ b/Engine/Plugins/PCG/Source/PCGEditor/Private/Widgets/AssetEditorViewport/SPCGEditorViewport.cpp
@@ -272,7 +272,7 @@ TSharedPtr<SWidget> SPCGEditorViewport::BuildViewportToolbar()
 		RightSection.AddEntry(UE::UnrealEd::CreateAssetViewerProfileSubmenu());
 		UE::AdvancedPreviewScene::Menus::ExtendAdvancedPreviewSceneSettings(
 			"PCG.ViewportToolbar.AssetViewerProfile",
-			UE::AdvancedPreviewScene::Menus::FSettingsOptions().ShowToggleGrid(false)
+			UE::AdvancedPreviewScene::Menus::FSettingsOptions()
 		);
 
 		RightSection.AddEntry(UE::PCGEditor::CreatePCGViewportSubmenu());
diff --git a/Engine/Source/Editor/AdvancedPreviewScene/Private/AdvancedPreviewScene.cpp b/Engine/Source/Editor/AdvancedPreviewScene/Private/AdvancedPreviewScene.cpp
--- a/Engine/Source/Editor/AdvancedPreviewScene/Private/AdvancedPreviewScene.cpp
+++ b/Engine/Source/Editor/AdvancedPreviewScene/Private/AdvancedPreviewScene.cpp
@@ -34,7 +34,8 @@ FAdvancedPreviewScene::FAdvancedPreviewScene(ConstructionValues CVS, float InFlo
 
 	RefreshDelegate = DefaultSettings->OnAssetViewerSettingsChanged().AddRaw(this, &FAdvancedPreviewScene::OnAssetViewerSettingsRefresh);
 
-	CurrentProfileIndex = DefaultSettings->Profiles.IsValidIndex(CurrentProfileIndex) ? GetDefault<UEditorPerProjectUserSettings>()->AssetViewerProfileIndex : 0;
+	const int32 SettingsProfileIndex = GetDefault<UEditorPerProjectUserSettings>()->AssetViewerProfileIndex;
+	CurrentProfileIndex = DefaultSettings->Profiles.IsValidIndex(SettingsProfileIndex) ? SettingsProfileIndex : 0;
 	ensureMsgf(DefaultSettings->Profiles.IsValidIndex(CurrentProfileIndex), TEXT("Invalid default settings pointer or current profile index"));
 	FPreviewSceneProfile& Profile = DefaultSettings->Profiles[CurrentProfileIndex];
 	Profile.LoadEnvironmentMap();
@@ -77,14 +78,16 @@ FAdvancedPreviewScene::FAdvancedPreviewScene(ConstructionValues CVS, float InFlo
 	PostProcessComponent->bUnbound = true;
 	AddComponent(PostProcessComponent, Transform);
 
-	UStaticMesh* FloorMesh = LoadObject<UStaticMesh>(NULL, TEXT("/Engine/EditorMeshes/AssetViewer/Floor_Mesh.Floor_Mesh"), NULL, LOAD_None, NULL);
-	check(FloorMesh);
 	FloorMeshComponent = NewObject<UStaticMeshComponent>(GetTransientPackage());
-	FloorMeshComponent->SetStaticMesh(FloorMesh);
 	FloorMeshComponent->bSelectable = false;
+	UStaticMesh* FloorMesh = Profile.GetEnvironmentFloorMesh();
+	SetFloorMesh(FloorMesh);
 
-	FTransform FloorTransform(FRotator(0, 0, 0), FVector(0, 0, -(InFloorOffset)), FVector(4.0f, 4.0f, 1.0f ));
-	AddComponent(FloorMeshComponent, FloorTransform);	
+	UMaterialInterface* FloorMaterial = Profile.GetEnvironmentFloorMaterial();
+	SetFloorMaterial(FloorMaterial);
+
+	FTransform FloorTransform(Profile.EnvironmentFloorRotation, FVector(0, 0, -(InFloorOffset)), Profile.EnvironmentFloorScale);
+	AddComponent(FloorMeshComponent, FloorTransform);
 
 	SetLightDirection(Profile.DirectionalLightRotation);
 
@@ -134,46 +137,62 @@ void FAdvancedPreviewScene::UpdateScene(FPreviewSceneProfile& Profile, bool bUpd
 		}
 	}
 
-	if (bUpdateEnvironment && SkyLight)
+	if (bUpdateEnvironment)
 	{
-		static const FName SkyBoxName("SkyBox");
-		static const FName CubeMapRotationName("CubemapRotation");
-
-		UTextureCube* EnvironmentTexture = Profile.EnvironmentCubeMap.LoadSynchronous();
-		UTexture* Texture = EnvironmentTexture;
-		InstancedSkyMaterial->GetTextureParameterValue(SkyBoxName, Texture);
-
-		if (Texture != EnvironmentTexture)
+		if (SkyLight)
 		{
-			InstancedSkyMaterial->SetTextureParameterValueEditorOnly(SkyBoxName, EnvironmentTexture);
-			SetSkyCubemap(EnvironmentTexture);
-			bSkyChanged = true;
-		}
+			static const FName SkyBoxName("SkyBox");
+			static const FName CubeMapRotationName("CubemapRotation");
+
+			UTextureCube* EnvironmentTexture = Profile.EnvironmentCubeMap.LoadSynchronous();
+			UTexture* Texture = EnvironmentTexture;
+			InstancedSkyMaterial->GetTextureParameterValue(SkyBoxName, Texture);
+
+			if (Texture != EnvironmentTexture)
+			{
+				InstancedSkyMaterial->SetTextureParameterValueEditorOnly(SkyBoxName, EnvironmentTexture);
+				SetSkyCubemap(EnvironmentTexture);
+				bSkyChanged = true;
+			}
 		
-		static const float OneOver360 = 1.0f / 360.0f;
-		float Rotation = Profile.LightingRigRotation;
-		InstancedSkyMaterial->GetScalarParameterValue(CubeMapRotationName, Rotation);
-		const bool bLightRighRotationChanged = !FMath::IsNearlyEqual(Rotation, Profile.LightingRigRotation, 0.05f);
-
-		// We also test any change of light direction because the default profile light orientation might not match the LightingRigRotation.
-		FRotator LightDir = GetLightDirection();
-		const bool bLightDirChanged = LightDir != Profile.DirectionalLightRotation;
-
-		if (bLightRighRotationChanged || bLightDirChanged)
-		{			
-			InstancedSkyMaterial->SetScalarParameterValueEditorOnly(CubeMapRotationName, Profile.LightingRigRotation * OneOver360);
-
-			// NOTE: this code is only executed when the UI rig rotation angle is interacted with.
-			// the l+mouse shortcut and the rig smooth rotation are handled at other places for the directional light
-			LightDir.Yaw += UILightingRigRotationDelta;
-			UILightingRigRotationDelta = 0;
-			SetLightDirection(LightDir);
-			DefaultSettings->Profiles[CurrentProfileIndex].DirectionalLightRotation = LightDir;
+			static const float OneOver360 = 1.0f / 360.0f;
+			float Rotation = Profile.LightingRigRotation;
+			InstancedSkyMaterial->GetScalarParameterValue(CubeMapRotationName, Rotation);
+			const bool bLightRighRotationChanged = !FMath::IsNearlyEqual(Rotation, Profile.LightingRigRotation, 0.05f);
+
+			// We also test any change of light direction because the default profile light orientation might not match the LightingRigRotation.
+			FRotator LightDir = GetLightDirection();
+			const bool bLightDirChanged = LightDir != Profile.DirectionalLightRotation;
+
+			if (bLightRighRotationChanged || bLightDirChanged)
+			{			
+				InstancedSkyMaterial->SetScalarParameterValueEditorOnly(CubeMapRotationName, Profile.LightingRigRotation * OneOver360);
+
+				// NOTE: this code is only executed when the UI rig rotation angle is interacted with.
+				// the l+mouse shortcut and the rig smooth rotation are handled at other places for the directional light
+				LightDir.Yaw += UILightingRigRotationDelta;
+				UILightingRigRotationDelta = 0;
+				SetLightDirection(LightDir);
+				DefaultSettings->Profiles[CurrentProfileIndex].DirectionalLightRotation = LightDir;
+
+				SkyLight->SourceCubemapAngle = Profile.LightingRigRotation;
+				bSkyChanged = true;
+			}
+		}
 
-			SkyLight->SourceCubemapAngle = Profile.LightingRigRotation;
-			bSkyChanged = true;
+		if (FloorMeshComponent)
+		{
+			UStaticMesh* FloorMesh = Profile.GetEnvironmentFloorMesh();
+			SetFloorMesh(FloorMesh);
+			FloorMeshComponent->SetRelativeScale3D(Profile.EnvironmentFloorScale);
+			FloorMeshComponent->SetRelativeRotation(Profile.EnvironmentFloorRotation);
+			UMaterialInterface* FloorMaterial = Profile.GetEnvironmentFloorMaterial();
+			for (int32 Index = 0; Index < FloorMeshComponent->GetNumMaterials(); Index++)
+			{
+				FloorMeshComponent->SetMaterial(Index, FloorMaterial);
+			}
 		}
-	}		
+	}
 
 	if (bUpdatePostProcessing)
 	{
@@ -211,18 +230,21 @@ FLinearColor FAdvancedPreviewScene::GetBackgroundColor() const
 
 float FAdvancedPreviewScene::GetFloorOffset() const
 {
-	return -(float)FloorMeshComponent->GetRelativeTransform().GetLocation().Z;
+	return -(float)FloorMeshComponent->GetRelativeLocation().Z;
 }
 
 void FAdvancedPreviewScene::SetFloorOffset(const float InFloorOffset)
 {
-	FTransform FloorTransform(FRotator(0, 0, 0), FVector(0, 0, -(InFloorOffset)), FVector(4.0f, 4.0f, 1.0f));
-
-	FloorMeshComponent->SetRelativeTransform(FloorTransform);
+	FloorMeshComponent->SetRelativeLocation(FVector(0, 0, -(InFloorOffset)));
 }
 
 void FAdvancedPreviewScene::SetProfileIndex(const int32 InProfileIndex)
 {
+	if (CurrentProfileIndex == InProfileIndex)
+	{
+		return;
+	}
+
 	CurrentProfileIndex = InProfileIndex;
 	DefaultSettings->Profiles[CurrentProfileIndex].LoadEnvironmentMap();
 	SetLightDirection(DefaultSettings->Profiles[CurrentProfileIndex].DirectionalLightRotation);
@@ -438,25 +460,20 @@ void FAdvancedPreviewScene::BindCommands()
 		FCanExecuteAction(),
 		FIsActionChecked::CreateRaw(this, &FAdvancedPreviewScene::IsEnvironmentEnabled)
 	);
-	
-	UICommandList->MapAction(Commands.ToggleGrid,
-		FExecuteAction::CreateRaw(this, &FAdvancedPreviewScene::HandleToggleGrid),
-		FCanExecuteAction(),
-		FIsActionChecked::CreateRaw(this, &FAdvancedPreviewScene::IsGridEnabled)
-	);
-	
-	// Override the grid toggle command so that toggling the grid toggles the setting first
-	UICommandList->MapAction(FShowFlagMenuCommands::Get().FindCommandForFlag(FEngineShowFlags::SF_Grid),
-		FExecuteAction::CreateRaw(this, &FAdvancedPreviewScene::HandleToggleGrid),
-		FCanExecuteAction(),
-		FIsActionChecked::CreateRaw(this, &FAdvancedPreviewScene::IsGridEnabled)
-	);
 
 	UICommandList->MapAction(Commands.TogglePostProcessing,
 		FExecuteAction::CreateRaw(this, &FAdvancedPreviewScene::HandleTogglePostProcessing),
 		FCanExecuteAction(),
 		FIsActionChecked::CreateRaw(this, &FAdvancedPreviewScene::IsPostProcessingEnabled)
 	);
+
+	UICommandList->MapAction(Commands.NextProfile,
+		FExecuteAction::CreateRaw(this, &FAdvancedPreviewScene::NavigateProfile, ENavigationDirection::Forward)
+	);
+
+	UICommandList->MapAction(Commands.PreviousProfile,
+		FExecuteAction::CreateRaw(this, &FAdvancedPreviewScene::NavigateProfile, ENavigationDirection::Backward)
+	);
 }
 
 TSharedRef<FUICommandList> FAdvancedPreviewScene::CreateCommandList()
@@ -486,16 +503,13 @@ bool FAdvancedPreviewScene::IsFloorEnabled() const
 
 void FAdvancedPreviewScene::HandleToggleGrid()
 {
-	FProperty* GridProperty = FindFProperty<FProperty>(FPreviewSceneProfile::StaticStruct(), GET_MEMBER_NAME_CHECKED(FPreviewSceneProfile, bShowGrid));
-	DefaultSettings->Profiles[CurrentProfileIndex].bShowGrid = !DefaultSettings->Profiles[CurrentProfileIndex].bShowGrid;
-
-	FPropertyChangedEvent PropertyEvent(GridProperty);
-	DefaultSettings->PostEditChangeProperty(PropertyEvent);
+	// DEPRECATED : Do nothing
 }
 
 bool FAdvancedPreviewScene::IsGridEnabled() const
 {
-	return DefaultSettings->Profiles[CurrentProfileIndex].bShowGrid;
+	// DEPRECATED
+	return false;
 }
 
 void FAdvancedPreviewScene::HandleTogglePostProcessing()
@@ -516,6 +530,48 @@ bool FAdvancedPreviewScene::IsPostProcessingEnabled() const
 	return DefaultSettings->Profiles[CurrentProfileIndex].bPostProcessingEnabled;
 }
 
+void FAdvancedPreviewScene::SetFloorMesh(UStaticMesh* InFloorMesh)
+{
+	if (FloorMeshComponent)
+	{
+		FloorMeshComponent->SetStaticMesh(InFloorMesh);
+	}
+}
+
+void FAdvancedPreviewScene::SetFloorMaterial(UMaterialInterface* InFloorMaterial, int32 InIndex)
+{
+	if (FloorMeshComponent)
+	{
+		if (InIndex == INDEX_NONE)
+		{
+			for (int32 Index = 0; Index < FloorMeshComponent->GetNumMaterials(); Index++)
+			{
+				FloorMeshComponent->SetMaterial(Index, InFloorMaterial);
+			}
+		}
+		else
+		{
+			FloorMeshComponent->SetMaterial(InIndex, InFloorMaterial);
+		}
+	}
+}
+
+void FAdvancedPreviewScene::SetFloorMeshRotation(const FRotator& InRotation)
+{
+	if (FloorMeshComponent)
+	{
+		FloorMeshComponent->SetRelativeRotation(InRotation);
+	}
+}
+
+void FAdvancedPreviewScene::SetFloorMeshScale(const FVector& InScale)
+{
+	if (FloorMeshComponent)
+	{
+		FloorMeshComponent->SetRelativeScale3D(InScale);
+	}
+}
+
 void FAdvancedPreviewScene::OnAssetViewerSettingsRefresh(const FName& InPropertyName)
 {
 	// If the profile was changed, update the current index and the scene.
@@ -534,7 +590,18 @@ void FAdvancedPreviewScene::OnAssetViewerSettingsRefresh(const FName& InProperty
 	{
 		const bool bNameNone = InPropertyName == NAME_None;
 
-		const bool bUpdateEnvironment = (InPropertyName == GET_MEMBER_NAME_CHECKED(FPreviewSceneProfile, EnvironmentCubeMap)) || (InPropertyName == GET_MEMBER_NAME_CHECKED(FPreviewSceneProfile, LightingRigRotation) || (InPropertyName == GET_MEMBER_NAME_CHECKED(UAssetViewerSettings, Profiles)));
+		static const TSet<FName> EnvironmentPropertiesName
+		{
+			GET_MEMBER_NAME_CHECKED(FPreviewSceneProfile, EnvironmentCubeMap),
+			GET_MEMBER_NAME_CHECKED(FPreviewSceneProfile, EnvironmentFloorMesh),
+			GET_MEMBER_NAME_CHECKED(FPreviewSceneProfile, EnvironmentFloorRotation),
+			GET_MEMBER_NAME_CHECKED(FPreviewSceneProfile, EnvironmentFloorScale),
+			GET_MEMBER_NAME_CHECKED(FPreviewSceneProfile, EnvironmentFloorMaterial),
+			GET_MEMBER_NAME_CHECKED(FPreviewSceneProfile, LightingRigRotation),
+			GET_MEMBER_NAME_CHECKED(UAssetViewerSettings, Profiles)
+		};
+
+		const bool bUpdateEnvironment = EnvironmentPropertiesName.Contains(InPropertyName);
 		const bool bUpdateSkyLight = bUpdateEnvironment || (InPropertyName == GET_MEMBER_NAME_CHECKED(FPreviewSceneProfile, SkyLightIntensity) || InPropertyName == GET_MEMBER_NAME_CHECKED(FPreviewSceneProfile, bUseSkyLighting) || (InPropertyName == GET_MEMBER_NAME_CHECKED(UAssetViewerSettings, Profiles)));
 		const bool bUpdateDirectionalLight = (InPropertyName == GET_MEMBER_NAME_CHECKED(FPreviewSceneProfile, DirectionalLightIntensity)) || (InPropertyName == GET_MEMBER_NAME_CHECKED(FPreviewSceneProfile, DirectionalLightColor));
 		const bool bUpdatePostProcessing = (InPropertyName == GET_MEMBER_NAME_CHECKED(FPreviewSceneProfile, PostProcessingSettings)) || (InPropertyName == GET_MEMBER_NAME_CHECKED(FPreviewSceneProfile, bPostProcessingEnabled));
@@ -547,3 +614,17 @@ void FAdvancedPreviewScene::OnAssetViewerSettingsRefresh(const FName& InProperty
 		OnProfileChangedDelegate.Broadcast(DefaultSettings->Profiles[CurrentProfileIndex], InPropertyName);
 	}
 }
+
+void FAdvancedPreviewScene::NavigateProfile(ENavigationDirection InDirection)
+{
+	const bool bMoveForward = InDirection == ENavigationDirection::Forward;
+	int32 NewProfileIndex = bMoveForward ? CurrentProfileIndex + 1 : CurrentProfileIndex - 1;
+	if (!DefaultSettings->Profiles.IsValidIndex(NewProfileIndex))
+	{
+		NewProfileIndex = bMoveForward ? 0 : DefaultSettings->Profiles.Num() - 1;
+	}
+	// This is to avoid a desync and hit ensure in FPreviewProfileController
+	GetMutableDefault<UEditorPerProjectUserSettings>()->AssetViewerProfileIndex = NewProfileIndex;
+	SetProfileIndex(NewProfileIndex);
+	OnProfileChangedDelegate.Broadcast(DefaultSettings->Profiles[NewProfileIndex], /** Proprety */NAME_None);
+}
diff --git a/Engine/Source/Editor/AdvancedPreviewScene/Private/AdvancedPreviewSceneCommands.cpp b/Engine/Source/Editor/AdvancedPreviewScene/Private/AdvancedPreviewSceneCommands.cpp
--- a/Engine/Source/Editor/AdvancedPreviewScene/Private/AdvancedPreviewSceneCommands.cpp
+++ b/Engine/Source/Editor/AdvancedPreviewScene/Private/AdvancedPreviewSceneCommands.cpp
@@ -12,8 +12,12 @@ void FAdvancedPreviewSceneCommands::RegisterCommands()
 {
 	UI_COMMAND(ToggleEnvironment, "Toggle Environment", "Toggles Environment visibility", EUserInterfaceActionType::ToggleButton, FInputChord(EKeys::I));
 	UI_COMMAND(ToggleFloor, "Toggle Floor", "Toggles floor visibility", EUserInterfaceActionType::ToggleButton, FInputChord(EKeys::O));
+	PRAGMA_DISABLE_DEPRECATION_WARNINGS
 	UI_COMMAND(ToggleGrid, "Toggle Grid", "Toggles grid visibility", EUserInterfaceActionType::ToggleButton, FInputChord(EModifierKey::Alt | EModifierKey::Shift, EKeys::G));
+	PRAGMA_ENABLE_DEPRECATION_WARNINGS
 	UI_COMMAND(TogglePostProcessing, "Toggle Post Processing", "Toggles whether Post Processing is enabled", EUserInterfaceActionType::ToggleButton, FInputChord(EKeys::P));
+	UI_COMMAND(NextProfile, "Next profile", "Navigate to the following profile", EUserInterfaceActionType::Button, FInputChord(EModifierKey::Alt | EModifierKey::Shift, EKeys::Down));
+	UI_COMMAND(PreviousProfile, "Previous profile", "Navigate to the preceding profile", EUserInterfaceActionType::Button, FInputChord(EModifierKey::Alt | EModifierKey::Shift, EKeys::Up));
 }
 
 #undef LOCTEXT_NAMESPACE
diff --git a/Engine/Source/Editor/AdvancedPreviewScene/Private/AdvancedPreviewSceneMenus.cpp b/Engine/Source/Editor/AdvancedPreviewScene/Private/AdvancedPreviewSceneMenus.cpp
--- a/Engine/Source/Editor/AdvancedPreviewScene/Private/AdvancedPreviewSceneMenus.cpp
+++ b/Engine/Source/Editor/AdvancedPreviewScene/Private/AdvancedPreviewSceneMenus.cpp
@@ -35,15 +35,6 @@ void ExtendAdvancedPreviewSceneSettings(FName InAssetViewerProfileSubmenuName, c
 		);
 	}
 	
-	if (InSettingsOptions.bShowToggleGrid)
-	{
-		Section.AddMenuEntry(
-        	FAdvancedPreviewSceneCommands::Get().ToggleGrid,
-        	LOCTEXT("ToggleGridLabel", "Grid"),
-        	LOCTEXT("ToggleGridTooltip", "Set the visibility of the preview scene's grid.")
-        );	
-	}
-	
 	if (InSettingsOptions.bShowToggleFloor)
 	{
 		Section.AddMenuEntry(
diff --git a/Engine/Source/Editor/AdvancedPreviewScene/Private/AssetViewerSettings.cpp b/Engine/Source/Editor/AdvancedPreviewScene/Private/AssetViewerSettings.cpp
--- a/Engine/Source/Editor/AdvancedPreviewScene/Private/AssetViewerSettings.cpp
+++ b/Engine/Source/Editor/AdvancedPreviewScene/Private/AssetViewerSettings.cpp
@@ -4,6 +4,7 @@
 #include "UObject/UnrealType.h"
 #include "Editor/EditorPerProjectUserSettings.h"
 #include "Editor.h"
+#include "Materials/MaterialInterface.h"
 
 #include UE_INLINE_GENERATED_CPP_BY_NAME(AssetViewerSettings)
 
@@ -11,6 +12,36 @@ FName UDefaultEditorProfiles::DefaultProfileName = FName("Epic Headquarters");
 FName UDefaultEditorProfiles::EditingProfileName = FName("Grey Wireframe");
 FName UDefaultEditorProfiles::GreyAmbientProfileName = FName("Grey Ambient");
 
+UStaticMesh* FPreviewSceneProfile::GetEnvironmentFloorMesh() const
+{
+	if (!EnvironmentFloorMesh && !EnvironmentFloorMeshPath.IsEmpty())
+	{
+		FPreviewSceneProfile* MutableThis = const_cast<FPreviewSceneProfile*>(this);
+		MutableThis->EnvironmentFloorMesh = LoadObject<UStaticMesh>(nullptr, *EnvironmentFloorMeshPath);
+	}
+
+	return EnvironmentFloorMesh.Get();
+}
+
+UMaterialInterface* FPreviewSceneProfile::GetEnvironmentFloorMaterial() const
+{
+	if (!EnvironmentFloorMaterial && !EnvironmentFloorMaterialPath.IsEmpty())
+	{
+		FPreviewSceneProfile* MutableThis = const_cast<FPreviewSceneProfile*>(this);
+		MutableThis->EnvironmentFloorMaterial = LoadObject<UMaterialInterface>(nullptr, *EnvironmentFloorMaterialPath);
+	}
+
+	return EnvironmentFloorMaterial.Get();
+}
+
+void FPreviewSceneProfile::PostEditChangeProperty(FPropertyChangedEvent& InEvent)
+{
+	// Store path to the set environment map texture
+	EnvironmentCubeMapPath = EnvironmentCubeMap.ToString();
+	EnvironmentFloorMeshPath = EnvironmentFloorMesh.ToString();
+	EnvironmentFloorMaterialPath = EnvironmentFloorMaterial.ToString();
+}
+
 const FPreviewSceneProfile* UDefaultEditorProfiles::GetProfile(const FString& ProfileName)
 {
 	for (const FPreviewSceneProfile& Profile : Profiles)
@@ -97,20 +128,47 @@ void UAssetViewerSettings::Save(bool bWarnIfFail)
 	SharedProfilesObject->TryUpdateDefaultConfigFile(FString(), bWarnIfFail);
 }
 
-void UAssetViewerSettings::PostEditChangeProperty(struct FPropertyChangedEvent& PropertyChangedEvent)
+void UAssetViewerSettings::PostEditChangeProperty(struct FPropertyChangedEvent& InPropertyChangedEvent)
 {
-	FName PropertyName = (PropertyChangedEvent.Property != nullptr) ? PropertyChangedEvent.Property->GetFName() : NAME_None;	
-	UObject* Outer = (PropertyChangedEvent.Property != nullptr) ? PropertyChangedEvent.Property->GetOwner<UObject>() : nullptr;
-	if (Outer != nullptr && ( Outer->GetName() == "PostProcessSettings" || Outer->GetName() == "WeightedBlendable" || Outer->GetName() == "Vector" || Outer->GetName() == "Vector4" || Outer->GetName() == "LinearColor"))
+	Super::PostEditChangeProperty(InPropertyChangedEvent);
+	
+	FName PropertyName = InPropertyChangedEvent.GetPropertyName();
+
+	const int32 ProfileIndex = GetMutableDefault<UEditorPerProjectUserSettings>()->AssetViewerProfileIndex;
+	if (Profiles.IsValidIndex(ProfileIndex))
 	{
-		PropertyName = GET_MEMBER_NAME_CHECKED(FPreviewSceneProfile, PostProcessingSettings);
+		Profiles[ProfileIndex].PostEditChangeProperty(InPropertyChangedEvent);
 	}
 
-	// Store path to the set enviroment map texture
-	if (PropertyName == GET_MEMBER_NAME_CHECKED(FPreviewSceneProfile, EnvironmentCubeMap))
-	{		
-		int32 ProfileIndex = GetMutableDefault<UEditorPerProjectUserSettings>()->AssetViewerProfileIndex;
-		Profiles[ProfileIndex].EnvironmentCubeMapPath = Profiles[ProfileIndex].EnvironmentCubeMap.ToString();
+	if (NumProfiles != Profiles.Num())
+	{
+		OnAssetViewerProfileAddRemovedEvent.Broadcast();
+		NumProfiles = Profiles.Num();
+	}
+
+	OnAssetViewerSettingsChangedEvent.Broadcast(PropertyName);
+}
+
+void UAssetViewerSettings::PostEditChangeChainProperty(FPropertyChangedChainEvent& InPropertyChangedEvent)
+{
+	Super::PostEditChangeChainProperty(InPropertyChangedEvent);
+
+	FName PropertyName = NAME_None;
+
+	// Get the inner property of Profiles
+	if (FEditPropertyChain::TDoubleLinkedListNode* ProfileNode = InPropertyChangedEvent.PropertyChain.GetActiveMemberNode())
+	{
+		ProfileNode = ProfileNode->GetNextNode();
+		if (ProfileNode)
+		{
+			PropertyName = ProfileNode->GetValue()->GetFName();
+		}
+	}
+
+	const int32 ProfileIndex = GetMutableDefault<UEditorPerProjectUserSettings>()->AssetViewerProfileIndex;
+	if (Profiles.IsValidIndex(ProfileIndex))
+	{
+		Profiles[ProfileIndex].PostEditChangeProperty(InPropertyChangedEvent);
 	}
 
 	if (NumProfiles != Profiles.Num())
@@ -166,22 +224,9 @@ void UAssetViewerSettings::PostInitProperties()
 	NumProfiles = Profiles.Num();
 
 	UEditorPerProjectUserSettings* ProjectSettings = GetMutableDefault<UEditorPerProjectUserSettings>();
-
-	// Find last saved profile index
-	int32 SelectedProfileIndex = INDEX_NONE;
-	for (int32 ProfileIndex = 0; ProfileIndex < Profiles.Num(); ++ProfileIndex)
-	{
-		if (Profiles[ProfileIndex].ProfileName == ProjectSettings->AssetViewerProfileName)
-		{
-			SelectedProfileIndex = ProfileIndex;
-			break;
-		}
-	}
-
-	// Only need to set profile index if we found a correct value since it defaults to 0
-	if (SelectedProfileIndex != INDEX_NONE)
+	if (!Profiles.IsValidIndex(ProjectSettings->AssetViewerProfileIndex))
 	{
-		ProjectSettings->AssetViewerProfileIndex = SelectedProfileIndex;
+		ProjectSettings->AssetViewerProfileIndex = 0;
 	}
 }
 
diff --git a/Engine/Source/Editor/AdvancedPreviewScene/Private/SAdvancedPreviewDetailsTab.cpp b/Engine/Source/Editor/AdvancedPreviewScene/Private/SAdvancedPreviewDetailsTab.cpp
--- a/Engine/Source/Editor/AdvancedPreviewScene/Private/SAdvancedPreviewDetailsTab.cpp
+++ b/Engine/Source/Editor/AdvancedPreviewScene/Private/SAdvancedPreviewDetailsTab.cpp
@@ -44,6 +44,11 @@ SAdvancedPreviewDetailsTab::~SAdvancedPreviewDetailsTab()
 		DefaultSettings->OnAssetViewerSettingsPostUndo().Remove(PostUndoDelegate);
 		DefaultSettings->Save();
 	}
+
+	if (const TSharedPtr<FAdvancedPreviewScene> PreviewScene = PreviewScenePtr.Pin())
+	{
+		PreviewScene->OnProfileChanged().RemoveAll(this);	
+	}
 }
 
 void SAdvancedPreviewDetailsTab::Construct(const FArguments& InArgs, const TSharedRef<FAdvancedPreviewScene>& InPreviewScene)
@@ -51,7 +56,7 @@ void SAdvancedPreviewDetailsTab::Construct(const FArguments& InArgs, const TShar
 	PreviewScenePtr = InPreviewScene;
 	DefaultSettings = UAssetViewerSettings::Get();
 	AdditionalSettings = InArgs._AdditionalSettings;
-	ProfileIndex = PerProjectSettings->AssetViewerProfileIndex;
+	ProfileIndex = DefaultSettings->Profiles.IsValidIndex(PerProjectSettings->AssetViewerProfileIndex) ? PerProjectSettings->AssetViewerProfileIndex : 0;
 	DetailCustomizations = InArgs._DetailCustomizations;
 	PropertyTypeCustomizations = InArgs._PropertyTypeCustomizations;
 	Delegates = InArgs._Delegates;
@@ -61,16 +66,6 @@ void SAdvancedPreviewDetailsTab::Construct(const FArguments& InArgs, const TShar
 	this->ChildSlot
 	[
 		SNew(SVerticalBox)
-		+ SVerticalBox::Slot()
-		.Padding(2.0f, 1.0f, 2.0f, 1.0f)
-		[
-			SNew(SHorizontalBox)
-			+SHorizontalBox::Slot()
-			[
-				SettingsView->AsShared()
-			]
-		]
-
 		+SVerticalBox::Slot()		
 		.Padding(2.0f, 1.0f, 2.0f, 1.0f)
 		.AutoHeight()
@@ -79,27 +74,14 @@ void SAdvancedPreviewDetailsTab::Construct(const FArguments& InArgs, const TShar
 			+SHorizontalBox::Slot()
 			.Padding(2.0f)
 			[
-				SNew(SHorizontalBox)
+				SAssignNew(ProfileComboBox, STextComboBox)
 				.ToolTipText(LOCTEXT("SceneProfileComboBoxToolTip", "Allows for switching between scene environment and lighting profiles."))
-				+ SHorizontalBox::Slot()
-				.Padding(0.0f, 0.0f, 2.0f, 0.0f)
-				.AutoWidth()
-				.VAlign(VAlign_Center)
-				[
-					SNew(STextBlock)
-					.Text(LOCTEXT("SceneProfileSettingsLabel", "Profile"))
-				]
-				+ SHorizontalBox::Slot()
-				.VAlign(VAlign_Fill)
-				[
-					SAssignNew(ProfileComboBox, STextComboBox)
-					.OptionsSource(&ProfileNames)
-					.OnSelectionChanged(this, &SAdvancedPreviewDetailsTab::ComboBoxSelectionChanged)				
-					.IsEnabled_Lambda([this]() -> bool
-					{
-						return ProfileNames.Num() > 1;
-					})
-				]
+				.OptionsSource(&ProfileNames)
+				.OnSelectionChanged(this, &SAdvancedPreviewDetailsTab::ComboBoxSelectionChanged)				
+				.IsEnabled_Lambda([this]() -> bool
+				{
+					return !ProfileNames.IsEmpty();
+				})
 			]
 			+ SHorizontalBox::Slot()
 			.Padding(2.0f)
@@ -113,6 +95,15 @@ void SAdvancedPreviewDetailsTab::Construct(const FArguments& InArgs, const TShar
 			+ SHorizontalBox::Slot()
 			.Padding(2.0f)
 			.AutoWidth()
+			[
+				SNew(SButton)
+				.OnClicked(this, &SAdvancedPreviewDetailsTab::CloneProfileButtonClick)
+				.Text(LOCTEXT("CloneProfileButton", "Clone Profile"))
+				.ToolTipText(LOCTEXT("SceneProfileCloneProfile", "Clones this profile into a new profile."))
+			]
+			+ SHorizontalBox::Slot()
+			.Padding(2.0f)
+			.AutoWidth()
 			[
 				SNew(SButton)
 				.OnClicked(this, &SAdvancedPreviewDetailsTab::RemoveOrResetProfileButtonClick)
@@ -138,10 +129,21 @@ void SAdvancedPreviewDetailsTab::Construct(const FArguments& InArgs, const TShar
 				})
 			]
 		]
+		+ SVerticalBox::Slot()
+		.Padding(2.0f, 1.0f, 2.0f, 1.0f)
+		[
+			SNew(SHorizontalBox)
+			+SHorizontalBox::Slot()
+			[
+				SettingsView->AsShared()
+			]
+		]
 	];
 
 	UpdateProfileNames();
 	UpdateSettingsView();
+
+	InPreviewScene->OnProfileChanged().AddSP(this, &SAdvancedPreviewDetailsTab::OnProfileChanged);
 }
 
 void SAdvancedPreviewDetailsTab::ComboBoxSelectionChanged(TSharedPtr<FString> NewSelection, ESelectInfo::Type SelectInfo)
@@ -150,13 +152,23 @@ void SAdvancedPreviewDetailsTab::ComboBoxSelectionChanged(TSharedPtr<FString> Ne
 
 	if (ProfileNames.Find(NewSelection, NewSelectionIndex))
 	{
-		ProfileIndex = NewSelectionIndex;
-		PerProjectSettings->AssetViewerProfileIndex = ProfileIndex;
-		UpdateSettingsView();	
-		check(PreviewScenePtr.IsValid());
-		if (SelectInfo == ESelectInfo::Type::OnMouseClick)
+		if (const TSharedPtr<FAdvancedPreviewScene> PreviewScene = PreviewScenePtr.Pin())
 		{
-			PreviewScenePtr.Pin()->SetProfileIndex(ProfileIndex);	
+			// No need to refresh anything
+			if (NewSelectionIndex == ProfileIndex 
+				&& ProfileIndex == PerProjectSettings->AssetViewerProfileIndex
+				&& PreviewScene->GetCurrentProfileIndex() == ProfileIndex)
+			{
+				return;
+			}
+
+			ProfileIndex = NewSelectionIndex;
+			PerProjectSettings->AssetViewerProfileIndex = ProfileIndex;
+			UpdateSettingsView();	
+			if (SelectInfo == ESelectInfo::Type::OnMouseClick)
+			{
+				PreviewScene->SetProfileIndex(ProfileIndex);	
+			}
 		}
 	}
 }
@@ -241,6 +253,56 @@ FReply SAdvancedPreviewDetailsTab::AddProfileButtonClick()
 	return FReply::Handled();
 }
 
+FReply SAdvancedPreviewDetailsTab::CloneProfileButtonClick()
+{
+	const FScopedTransaction Transaction(LOCTEXT("CloneSceneProfile", "Cloning Preview Scene Profile"));
+    DefaultSettings->Modify();
+
+    // Add new profile to settings instance based on current profile
+	const int32 CurrentProfileIndex = DefaultSettings->Profiles.IsValidIndex(ProfileIndex) ? ProfileIndex : 0;
+	FPreviewSceneProfile CurrentProfile = DefaultSettings->Profiles[CurrentProfileIndex];
+	CurrentProfile.bIsEngineDefaultProfile = false;
+	CurrentProfile.bSharedProfile = false;
+    ProfileIndex = DefaultSettings->Profiles.Add(CurrentProfile);
+    FPreviewSceneProfile& NewProfile = DefaultSettings->Profiles.Last();
+
+    // Try to create a valid profile name when one is added
+    bool bFoundValidName = false;
+    int32 ProfileAppendNum = FMath::Max(0, DefaultSettings->Profiles.Num() - 1);
+    FString NewProfileName;
+    while (!bFoundValidName)
+    {
+    	NewProfileName = FString::Printf(TEXT("Profile_%i"), ProfileAppendNum);
+    
+    	bool bValidName = true;
+    	for (const FPreviewSceneProfile& Profile : DefaultSettings->Profiles)
+    	{
+    		if (Profile.ProfileName == NewProfileName)
+    		{
+    			bValidName = false;				
+    			break;
+    		}
+    	}
+    
+    	if (!bValidName)
+    	{
+    		++ProfileAppendNum;
+    	}
+    
+    	bFoundValidName = bValidName;
+    }
+
+    NewProfile.ProfileName = NewProfileName;
+    PerProjectSettings->AssetViewerProfileIndex = ProfileIndex;
+    DefaultSettings->PostEditChange();
+
+    // Change selection to new profile so the user directly sees the profile that was added
+    Refresh();
+    ProfileComboBox->SetSelectedItem(ProfileNames.Last());
+
+    return FReply::Handled();
+}
+
 FReply SAdvancedPreviewDetailsTab::RemoveOrResetProfileButtonClick()
 {
 	const FPreviewSceneProfile& CurrentProfile = DefaultSettings->Profiles[ProfileIndex];
@@ -337,11 +399,26 @@ void SAdvancedPreviewDetailsTab::OnAssetViewerSettingsPostUndo()
 	PreviewScenePtr.Pin()->UpdateScene(DefaultSettings->Profiles[ProfileIndex]);
 }
 
+void SAdvancedPreviewDetailsTab::OnProfileChanged(const FPreviewSceneProfile& InPreviewSceneProfile, FName InName)
+{
+	if (PerProjectSettings->AssetViewerProfileIndex != ProfileIndex)
+	{
+		Refresh();
+	}
+}
+
 void SAdvancedPreviewDetailsTab::OnPreviewSceneChanged(TSharedRef<FAdvancedPreviewScene> PreviewScene)
 {
+	if (const TSharedPtr<FAdvancedPreviewScene> CurrentPreviewScene = PreviewScenePtr.Pin())
+	{
+		CurrentPreviewScene->OnProfileChanged().RemoveAll(this);
+	}
+
   	PreviewScenePtr = PreviewScene;
 	Refresh();
 
+	PreviewScene->OnProfileChanged().AddSP(this, &SAdvancedPreviewDetailsTab::OnProfileChanged);
+
 }
 
 #undef LOCTEXT_NAMESPACE
\ No newline at end of file
diff --git a/Engine/Source/Editor/AdvancedPreviewScene/Public/AdvancedPreviewScene.h b/Engine/Source/Editor/AdvancedPreviewScene/Public/AdvancedPreviewScene.h
--- a/Engine/Source/Editor/AdvancedPreviewScene/Public/AdvancedPreviewScene.h
+++ b/Engine/Source/Editor/AdvancedPreviewScene/Public/AdvancedPreviewScene.h
@@ -82,9 +82,11 @@ public:
 	UE_API bool IsFloorEnabled() const;
 
 	/** Toggle the grid on and off. */
+	UE_DEPRECATED(5.8, "bShowGrid is deprecated, this function does nothing")
 	UE_API void HandleToggleGrid();
 	
 	/** Whether the grid is enabled. */
+	UE_DEPRECATED(5.8, "bShowGrid is deprecated, this function returns false")
 	UE_API bool IsGridEnabled() const;
 
 	/** Toggle post processing on and off */
@@ -92,6 +94,18 @@ public:
 	
 	/** Whether post processing is enabled */
 	UE_API bool IsPostProcessingEnabled() const;
+	
+	/** Sets the mesh for the floor */
+	UE_API void SetFloorMesh(UStaticMesh* InFloorMesh);
+	
+	/** Sets the floor material on a specific slot, or -1 for all slots */
+	UE_API void SetFloorMaterial(UMaterialInterface* InFloorMaterial, int32 InIndex = -1);
+
+	/** Sets the floor mesh rotation */
+	UE_API void SetFloorMeshRotation(const FRotator& InRotation);
+
+	/** Sets the floor mesh scale */
+	UE_API void SetFloorMeshScale(const FVector& InScale);
 
 protected:
 	/** Create and map the command list. */
@@ -102,6 +116,14 @@ protected:
 	
 	/** Handle refreshing the scene when settings change */
 	UE_API void OnAssetViewerSettingsRefresh(const FName& InPropertyName);
+	
+	/** Navigate to the next or previous profile */
+	enum class ENavigationDirection : uint8
+	{
+		Forward,
+		Backward,
+	};
+	UE_API void NavigateProfile(ENavigationDirection InDirection);
 
 protected:
 	FOnProfileChanged OnProfileChangedDelegate;
diff --git a/Engine/Source/Editor/AdvancedPreviewScene/Public/AdvancedPreviewSceneCommands.h b/Engine/Source/Editor/AdvancedPreviewScene/Public/AdvancedPreviewSceneCommands.h
--- a/Engine/Source/Editor/AdvancedPreviewScene/Public/AdvancedPreviewSceneCommands.h
+++ b/Engine/Source/Editor/AdvancedPreviewScene/Public/AdvancedPreviewSceneCommands.h
@@ -12,7 +12,9 @@
 
 class FUICommandInfo;
 
-class ADVANCEDPREVIEWSCENE_API FAdvancedPreviewSceneCommands : public TCommands<FAdvancedPreviewSceneCommands>
+PRAGMA_DISABLE_DEPRECATION_WARNINGS // For ToggleGrid
+class ADVANCEDPREVIEWSCENE_API FAdvancedPreviewSceneCommands 
+	: public TCommands<FAdvancedPreviewSceneCommands>
 {
 
 public:
@@ -30,15 +32,23 @@ public:
 
 	/** Toggles floor visibility */
 	TSharedPtr< FUICommandInfo > ToggleFloor;
-
-	/** Toggles the grid */
-	TSharedPtr< FUICommandInfo > ToggleGrid;
 	
 	/** Toggles post processing */
 	TSharedPtr< FUICommandInfo > TogglePostProcessing;
+	
+	/** Toggles the grid */
+	UE_DEPRECATED(5.8, "No longer in use, will be removed")
+	TSharedPtr< FUICommandInfo > ToggleGrid;
+	
+	/** Navigate to next profile */
+	TSharedPtr<FUICommandInfo> NextProfile;
+	
+	/** Navigate to previous profile */
+	TSharedPtr<FUICommandInfo> PreviousProfile;
 
 	/**
 	 * Initialize commands
 	 */
 	virtual void RegisterCommands() override;
 };
+PRAGMA_ENABLE_DEPRECATION_WARNINGS // For ToggleGrid
diff --git a/Engine/Source/Editor/AdvancedPreviewScene/Public/AdvancedPreviewSceneMenus.h b/Engine/Source/Editor/AdvancedPreviewScene/Public/AdvancedPreviewSceneMenus.h
--- a/Engine/Source/Editor/AdvancedPreviewScene/Public/AdvancedPreviewSceneMenus.h
+++ b/Engine/Source/Editor/AdvancedPreviewScene/Public/AdvancedPreviewSceneMenus.h
@@ -16,6 +16,7 @@ namespace UE::AdvancedPreviewScene::Menus
 struct FSettingsOptions
 {
 	bool bShowToggleEnvironment = true;
+	UE_DEPRECATED(5.8, "No longer in use, will be removed")
 	bool bShowToggleGrid = true;
 	bool bShowToggleFloor = true;
 	bool bShowTogglePostProcessing = true;
@@ -26,9 +27,9 @@ struct FSettingsOptions
 		return *this;
 	}
 
+	UE_DEPRECATED(5.8, "No longer in use, will be removed")
 	FSettingsOptions& ShowToggleGrid(bool bInShowToggleGrid)
 	{
-		bShowToggleGrid = bInShowToggleGrid;
 		return *this;
 	}
 
diff --git a/Engine/Source/Editor/AdvancedPreviewScene/Public/AssetViewerSettings.h b/Engine/Source/Editor/AdvancedPreviewScene/Public/AssetViewerSettings.h
--- a/Engine/Source/Editor/AdvancedPreviewScene/Public/AssetViewerSettings.h
+++ b/Engine/Source/Editor/AdvancedPreviewScene/Public/AssetViewerSettings.h
@@ -15,6 +15,7 @@
 
 #define UE_API ADVANCEDPREVIEWSCENE_API
 
+PRAGMA_DISABLE_DEPRECATION_WARNINGS // For bShowGrid
 /**
 * Preview scene profile settings structure.
 */
@@ -41,11 +42,19 @@ struct FPreviewSceneProfile
 		// Set up default cube map texture from editor/engine textures
 		EnvironmentCubeMap = nullptr;
 		EnvironmentCubeMapPath = TEXT("/Engine/EditorMaterials/AssetViewer/EpicQuadPanorama_CC+EV1.EpicQuadPanorama_CC+EV1");
+		EnvironmentFloorMesh = nullptr;
+		EnvironmentFloorMeshPath = TEXT("/Engine/EditorMeshes/AssetViewer/Floor_Mesh.Floor_Mesh");
+		EnvironmentFloorScale = FVector(4.0f, 4.0f, 1.0f);
+		EnvironmentFloorRotation = FRotator();
 		bPostProcessingEnabled = true;
 		DirectionalLightRotation = FRotator(-40.f, -67.5f, 0.f);
 		bEnableToneMapping = true;
 		bShowMeshEdges = false;
-		bShowGrid = false;
+		// This will stop scene from becoming bright due to exposure
+		PostProcessingSettings.bOverride_AutoExposureMinBrightness = true;
+		PostProcessingSettings.AutoExposureMinBrightness = -1;
+		PostProcessingSettings.bOverride_AutoExposureMaxBrightness = true;
+		PostProcessingSettings.AutoExposureMaxBrightness = 1;
 	}
 
 	/** Name to identify the profile */
@@ -85,12 +94,13 @@ struct FPreviewSceneProfile
 	bool bShowEnvironment;
 
 	/** Toggle visibility of the floor mesh */
-	UPROPERTY(EditAnywhere, config, Category = Environment)
+	UPROPERTY(EditAnywhere, config, Category = Environment, meta=(DisplayAfter="EnvironmentCubeMap"))
 	bool bShowFloor;
 
 	/** Toggle visibility of floor grid on/off */
-	UPROPERTY(EditAnywhere, config, Category = Environment)
-	bool bShowGrid;
+	UE_DEPRECATED(5.8, "Removed from preview scene profile since it's not relevant to every viewport")
+	UPROPERTY(config)
+	bool bShowGrid = false;
 
 	/** The environment color, used if Show Environment is false. */
 	UPROPERTY(EditAnywhere, config, Category = Environment, meta=(EditCondition="!bShowEnvironment"))
@@ -101,13 +111,37 @@ struct FPreviewSceneProfile
 	float EnvironmentIntensity;
 
 	/** Sets environment cube map used for sky lighting and reflections */
-	UPROPERTY(EditAnywhere, transient, Category = Environment)
+	UPROPERTY(EditAnywhere, transient, Category = Environment, meta=(EditCondition="bShowEnvironment", AssetViewInitialViewType=1, AssetViewInitialThumbnailSize=4, AssetViewSettingsName="", ForceShowPluginContent, ForceShowEngineContent))
 	TSoftObjectPtr<UTextureCube> EnvironmentCubeMap;
 
 	/** Storing path to environment cube to prevent it from getting cooked */
 	UPROPERTY(config)
 	FString EnvironmentCubeMapPath;
 
+	/** Sets the environment floor mesh */
+	UPROPERTY(EditAnywhere, transient, Category = Environment, meta=(EditCondition="bShowFloor", AssetViewInitialViewType=1, AssetViewInitialThumbnailSize=4, AssetViewSettingsName="", ForceShowPluginContent, ForceShowEngineContent))
+	TSoftObjectPtr<UStaticMesh> EnvironmentFloorMesh;
+
+	/** Storing path to the environment floor mesh to prevent it from getting cooked */
+	UPROPERTY(config)
+	FString EnvironmentFloorMeshPath;
+
+	/** Sets the rotation of the floor mesh */
+	UPROPERTY(EditAnywhere, config, Category = Environment, meta=(EditCondition="bShowFloor"))
+	FRotator EnvironmentFloorRotation;
+
+	/** Sets the scale of the floor mesh */
+	UPROPERTY(EditAnywhere, config, Category = Environment, meta=(EditCondition="bShowFloor"))
+	FVector EnvironmentFloorScale;
+
+	/** Sets the environment floor material */
+	UPROPERTY(EditAnywhere, transient, Category = Environment, meta=(EditCondition="bShowFloor", AssetViewInitialViewType=1, AssetViewInitialThumbnailSize=4, AssetViewSettingsName="", ForceShowPluginContent, ForceShowEngineContent))
+	TSoftObjectPtr<UMaterialInterface> EnvironmentFloorMaterial;
+
+	/** Storing path to the environment floor material to prevent it from getting cooked */
+	UPROPERTY(config)
+	FString EnvironmentFloorMaterialPath;
+
 	/** Whether or not the Post Processing should influence the scene */
 	UPROPERTY(EditAnywhere, config, Category = PostProcessing, AdvancedDisplay)
 	bool bPostProcessingEnabled;
@@ -173,11 +207,36 @@ struct FPreviewSceneProfile
 		}
 
 		ShowFlags.SetTonemapper(bEnableToneMapping);
-		ShowFlags.SetGrid(bShowGrid);
 		ShowFlags.SetMeshEdges(bShowMeshEdges);
 	}
+	
+	/** Loads and returns the environment floor mesh */
+	UE_API UStaticMesh* GetEnvironmentFloorMesh() const;
+	
+	/** Loads and returns the environment floor material */
+	UE_API UMaterialInterface* GetEnvironmentFloorMaterial() const;
+	
+	/** Called when a profile property has been changed */
+	void PostEditChangeProperty(FPropertyChangedEvent& InEvent);
+	
+	bool operator==(const FPreviewSceneProfile& InOther) const
+	{
+		return ProfileName.Equals(InOther.ProfileName);
+	}
+
+	bool operator!=(const FPreviewSceneProfile& InOther) const
+	{
+		return !(*this == InOther);
+	}
+
+	friend uint32 GetTypeHash(const FPreviewSceneProfile& InItem)
+	{
+		return GetTypeHash(InItem.ProfileName);
+	}
 };
 
+PRAGMA_ENABLE_DEPRECATION_WARNINGS // For bShowGrid
+
 UCLASS(MinimalAPI)
 class UDefaultEditorProfiles : public UObject
 {
@@ -198,7 +257,6 @@ public:
 		EditingProfile.bSharedProfile = true;
 		EditingProfile.bShowEnvironment = false;
 		EditingProfile.bShowFloor = false;
-		EditingProfile.bShowGrid = true;
 		EditingProfile.EnvironmentColor = FLinearColor::MakeFromHSV8(0,0,10);
 		EditingProfile.bUseSkyLighting = true;
 		EditingProfile.bPostProcessingEnabled = false;
@@ -212,7 +270,6 @@ public:
 		GreyAmbientProfile.bSharedProfile = true;
 		GreyAmbientProfile.bShowEnvironment = true;
 		GreyAmbientProfile.bShowFloor = true;
-		GreyAmbientProfile.bShowGrid = true;
 		GreyAmbientProfile.bUseSkyLighting = true;
 		GreyAmbientProfile.bPostProcessingEnabled = false;
 		GreyAmbientProfile.bShowMeshEdges = false;
@@ -285,7 +342,8 @@ public:
 	FOnAssetViewerSettingsPostUndo& OnAssetViewerSettingsPostUndo() { return OnAssetViewerSettingsPostUndoEvent; }
 
 	/** Begin UObject */
-	UE_API virtual void PostEditChangeProperty(struct FPropertyChangedEvent& PropertyChangedEvent) override;
+	UE_API virtual void PostEditChangeProperty(struct FPropertyChangedEvent& InPropertyChangedEvent) override;
+	UE_API virtual void PostEditChangeChainProperty(struct FPropertyChangedChainEvent& InPropertyChangedEvent) override;
 	UE_API virtual void PostInitProperties() override;
 	/** End UObject */
 
diff --git a/Engine/Source/Editor/AdvancedPreviewScene/Public/SAdvancedPreviewDetailsTab.h b/Engine/Source/Editor/AdvancedPreviewScene/Public/SAdvancedPreviewDetailsTab.h
--- a/Engine/Source/Editor/AdvancedPreviewScene/Public/SAdvancedPreviewDetailsTab.h
+++ b/Engine/Source/Editor/AdvancedPreviewScene/Public/SAdvancedPreviewDetailsTab.h
@@ -8,6 +8,7 @@
 #include "Widgets/SCompoundWidget.h"
 #include "PropertyEditorDelegates.h"
 #include "AdvancedPreviewSceneModule.h"
+#include "AssetViewerSettings.h"
 
 #define UE_API ADVANCEDPREVIEWSCENE_API
 
@@ -37,10 +38,8 @@ public:
 
 	SLATE_END_ARGS()
 
-public:
-	/** **/
 	UE_API SAdvancedPreviewDetailsTab();
-	UE_API ~SAdvancedPreviewDetailsTab();
+	UE_API virtual ~SAdvancedPreviewDetailsTab() override;
 
 	/** SWidget functions */
 	UE_API void Construct(const FArguments& InArgs, const TSharedRef<FAdvancedPreviewScene>& InPreviewScene);
@@ -53,11 +52,14 @@ protected:
 	UE_API void UpdateSettingsView();
 	UE_API void UpdateProfileNames();
 	UE_API FReply AddProfileButtonClick();
+	UE_API FReply CloneProfileButtonClick();
 	UE_API FReply RemoveOrResetProfileButtonClick();
-protected:
+
 	UE_API void OnAssetViewerSettingsRefresh(const FName& InPropertyName);
 	UE_API void OnAssetViewerSettingsPostUndo();
-protected:
+
+	UE_API void OnProfileChanged(const FPreviewSceneProfile& InPreviewSceneProfile, FName InName);
+
 	/** Property viewing widget */
 	TSharedPtr<IDetailsView> SettingsView;
 	TSharedPtr<STextComboBox> ProfileComboBox;
diff --git a/Engine/Source/Editor/MaterialEditor/Private/SMaterialEditorViewport.cpp b/Engine/Source/Editor/MaterialEditor/Private/SMaterialEditorViewport.cpp
--- a/Engine/Source/Editor/MaterialEditor/Private/SMaterialEditorViewport.cpp
+++ b/Engine/Source/Editor/MaterialEditor/Private/SMaterialEditorViewport.cpp
@@ -1007,7 +1007,7 @@ TSharedPtr<SWidget> SMaterialEditor3DPreviewViewport::BuildViewportToolbar()
 				UE::MaterialEditor::ExtendPreviewSceneSettingsSubmenu(AssetViewerProfileMenuName);
 				UE::AdvancedPreviewScene::Menus::ExtendAdvancedPreviewSceneSettings(
 					AssetViewerProfileMenuName,
-					UE::AdvancedPreviewScene::Menus::FSettingsOptions().ShowToggleGrid(false)
+					UE::AdvancedPreviewScene::Menus::FSettingsOptions()
 				);
 				UE::UnrealEd::ExtendPreviewSceneSettingsWithTabEntry(AssetViewerProfileMenuName);
 			}
diff --git a/Engine/Source/Editor/PropertyEditor/Private/UserInterface/PropertyEditor/SPropertyMenuAssetPicker.cpp b/Engine/Source/Editor/PropertyEditor/Private/UserInterface/PropertyEditor/SPropertyMenuAssetPicker.cpp
--- a/Engine/Source/Editor/PropertyEditor/Private/UserInterface/PropertyEditor/SPropertyMenuAssetPicker.cpp
+++ b/Engine/Source/Editor/PropertyEditor/Private/UserInterface/PropertyEditor/SPropertyMenuAssetPicker.cpp
@@ -37,6 +37,47 @@ void SPropertyMenuAssetPicker::Construct( const FArguments& InArgs )
 	OnSet = InArgs._OnSet;
 	OnClose = InArgs._OnClose;
 
+	// Get initial view type
+	EAssetViewType::Type AssetViewType = EAssetViewType::List;
+	static const FName AssetViewTypeMetadata = TEXT("AssetViewInitialViewType");
+	if (PropertyHandle && PropertyHandle->HasMetaData(AssetViewTypeMetadata))
+	{
+		const int32 Value = PropertyHandle->GetIntMetaData(AssetViewTypeMetadata);
+		if (Value >= 0 && Value <= (EAssetViewType::Type::MAX - 1))
+		{
+			AssetViewType = static_cast<EAssetViewType::Type>(Value);
+		}
+	}
+
+	// Get initial thumbnail size
+	EThumbnailSize AssetThumbnailSize = EThumbnailSize::Medium;
+	static const FName AssetThumbnailSizeMetadata = TEXT("AssetViewInitialThumbnailSize");
+	if (PropertyHandle && PropertyHandle->HasMetaData(AssetThumbnailSizeMetadata))
+	{
+		const int32 Value = PropertyHandle->GetIntMetaData(AssetThumbnailSizeMetadata);
+		if (Value >= 0 && Value <= (static_cast<int32>(EThumbnailSize::MAX) - 1))
+		{
+			AssetThumbnailSize = static_cast<EThumbnailSize>(Value);
+		}
+	}
+
+	// Get initial thumbnail scale
+	TOptional<float> AssetThumbnailScale;
+	static const FName AssetThumbnailScaleMetadata = TEXT("AssetViewInitialThumbnailScale");
+	if (PropertyHandle && PropertyHandle->HasMetaData(AssetThumbnailScaleMetadata))
+	{
+		const int32 Value = PropertyHandle->GetFloatMetaData(AssetThumbnailScaleMetadata);
+		if (Value >= 0 && Value <= 1)
+		{
+			AssetThumbnailScale = Value;
+		}
+	}
+
+	// Get settings to load/save, use empty string to not use any settings
+	static const FName AssetViewSettingsNameMetadata = TEXT("AssetViewSettingsName");
+	const FString SaveSettingsName = PropertyHandle && PropertyHandle->HasMetaData(AssetViewSettingsNameMetadata) ? PropertyHandle->GetMetaData(AssetViewSettingsNameMetadata) : TEXT("AssetPropertyPicker");
+
+	// Get show flags
 	const bool bForceShowEngineContent = PropertyHandle ? PropertyHandle->HasMetaData(TEXT("ForceShowEngineContent")) : false;
 	const bool bForceShowPluginContent = PropertyHandle ? PropertyHandle->HasMetaData(TEXT("ForceShowPluginContent")) : false;
 
@@ -140,12 +181,17 @@ void SPropertyMenuAssetPicker::Construct( const FArguments& InArgs )
 		// Set a delegate for setting the asset from the picker via the keyboard
 		AssetPickerConfig.OnAssetEnterPressed = FOnAssetEnterPressed::CreateSP(this, &SPropertyMenuAssetPicker::OnAssetEnterPressed);
 		// Use the list view by default
-		AssetPickerConfig.InitialAssetViewType = EAssetViewType::List;
+		AssetPickerConfig.InitialAssetViewType = AssetViewType;
 		// TODO: Remove "PropertyEditorHelpers.h" include when IsNewStyleEnabled will be removed
 		if (UE::Editor::ContentBrowser::IsNewStyleEnabled())
 		{
 			// Use the medium size by default if the new style is enabled
-			AssetPickerConfig.InitialThumbnailSize = EThumbnailSize::Medium;
+			AssetPickerConfig.InitialThumbnailSize = AssetThumbnailSize;
+		}
+		// Set scale for thumbnail
+		if (AssetThumbnailScale.IsSet())
+		{
+			AssetPickerConfig.ThumbnailScale = AssetThumbnailScale.GetValue();
 		}
 		// The initial selection should be the current value
 		AssetPickerConfig.InitialAssetSelection = CurrentObject;
@@ -158,7 +204,7 @@ void SPropertyMenuAssetPicker::Construct( const FArguments& InArgs )
 		// Don't allow dragging
 		AssetPickerConfig.bAllowDragging = false;
 		// Save the settings into a special section for asset pickers for properties
-		AssetPickerConfig.SaveSettingsName = TEXT("AssetPropertyPicker");
+		AssetPickerConfig.SaveSettingsName = SaveSettingsName;
 		// Populate the referencing assets via property handle
 		AssetPickerConfig.PropertyHandle = PropertyHandle;
 		// Populate the additional referencing assets with the Owner asset data
diff --git a/Engine/Source/Editor/StaticMeshEditor/Private/SStaticMeshEditorViewport.cpp b/Engine/Source/Editor/StaticMeshEditor/Private/SStaticMeshEditorViewport.cpp
--- a/Engine/Source/Editor/StaticMeshEditor/Private/SStaticMeshEditorViewport.cpp
+++ b/Engine/Source/Editor/StaticMeshEditor/Private/SStaticMeshEditorViewport.cpp
@@ -977,9 +977,9 @@ void SStaticMeshEditorViewport::BindCommands()
 
 	CommandList->MapAction(
 		Commands.SetShowGrid,
-		FExecuteAction::CreateRaw( PreviewScene.Get(), &FAdvancedPreviewScene::HandleToggleGrid ),
+		FExecuteAction::CreateSP( EditorViewportClientRef, &FStaticMeshEditorViewportClient::SetShowGrid ),
 		FCanExecuteAction(),
-		FIsActionChecked::CreateRaw( PreviewScene.Get(), &FAdvancedPreviewScene::IsGridEnabled ) );
+		FIsActionChecked::CreateSP( EditorViewportClientRef, &FStaticMeshEditorViewportClient::IsSetShowGridChecked ) );
 
 	CommandList->MapAction(
 		Commands.SetShowBounds,
diff --git a/Engine/Source/Editor/UnrealEd/Classes/Editor/EditorPerProjectUserSettings.h b/Engine/Source/Editor/UnrealEd/Classes/Editor/EditorPerProjectUserSettings.h
--- a/Engine/Source/Editor/UnrealEd/Classes/Editor/EditorPerProjectUserSettings.h
+++ b/Engine/Source/Editor/UnrealEd/Classes/Editor/EditorPerProjectUserSettings.h
@@ -185,12 +185,9 @@ class UEditorPerProjectUserSettings : public UObject
 	
 public:
 	// Per project user settings for which asset viewer profile should be used
-	UPROPERTY()
+	UPROPERTY(Config)
 	int32 AssetViewerProfileIndex;
 
-	UPROPERTY(config)
-	FString AssetViewerProfileName;
-
 	/** The feature level we should use when loading or creating a new world */
 	UPROPERTY(config)
 	int32 PreviewFeatureLevel;
