EditorUsability : PreviewScene
- Fixed some issues related to new featues introduced recently (feedback)
- Fixed asset viewports should not all use the same preview scene profile (unsync)
- Added notification to let user know profile changed when using shortcuts
- Added tooltips to let user know about keybind shortcuts available
- Added confirmation dialog for remove/reset profile actions

#Synced-CL 50871375

diff --git a/Engine/Source/Editor/AdvancedPreviewScene/Private/AdvancedPreviewScene.cpp b/Engine/Source/Editor/AdvancedPreviewScene/Private/AdvancedPreviewScene.cpp
--- a/Engine/Source/Editor/AdvancedPreviewScene/Private/AdvancedPreviewScene.cpp
+++ b/Engine/Source/Editor/AdvancedPreviewScene/Private/AdvancedPreviewScene.cpp
@@ -23,8 +23,12 @@
 #include "AdvancedPreviewSceneCommands.h"
 #include "Framework/Commands/UICommandList.h"
 #include "Framework/Application/SlateApplication.h"
+#include "Framework/Notifications/NotificationManager.h"
 #include "Editor.h"
 #include "ShowFlagMenuCommands.h"
+#include "Widgets/Notifications/SNotificationList.h"
+
+#define LOCTEXT_NAMESPACE "AdvancedPreviewScene"
 
 FAdvancedPreviewScene::FAdvancedPreviewScene(ConstructionValues CVS, float InFloorOffset)
 	: FPreviewScene(CVS)
@@ -627,4 +631,12 @@ void FAdvancedPreviewScene::NavigateProfile(ENavigationDirection InDirection)
 	GetMutableDefault<UEditorPerProjectUserSettings>()->AssetViewerProfileIndex = NewProfileIndex;
 	SetProfileIndex(NewProfileIndex);
 	OnProfileChangedDelegate.Broadcast(DefaultSettings->Profiles[NewProfileIndex], /** Proprety */NAME_None);
+
+	// Notify user
+	FSlateNotificationManager& NotificationManager = FSlateNotificationManager::Get();
+	FNotificationInfo Info(FText::Format(LOCTEXT("NavigateAdvancedPreviewSceneProfile", "Applying scene profile: {0}"), FText::FromString(DefaultSettings->Profiles[NewProfileIndex].ProfileName)));
+	Info.ExpireDuration = 4.f;
+	NotificationManager.AddNotification(Info);
 }
+
+#undef LOCTEXT_NAMESPACE
diff --git a/Engine/Source/Editor/AdvancedPreviewScene/Private/SAdvancedPreviewDetailsTab.cpp b/Engine/Source/Editor/AdvancedPreviewScene/Private/SAdvancedPreviewDetailsTab.cpp
--- a/Engine/Source/Editor/AdvancedPreviewScene/Private/SAdvancedPreviewDetailsTab.cpp
+++ b/Engine/Source/Editor/AdvancedPreviewScene/Private/SAdvancedPreviewDetailsTab.cpp
@@ -15,6 +15,7 @@
 
 #include "AdvancedPreviewScene.h"
 #include "IDetailRootObjectCustomization.h"
+#include "Misc/MessageDialog.h"
 
 #define LOCTEXT_NAMESPACE "SPrettyPreview"
 
@@ -44,11 +45,6 @@ SAdvancedPreviewDetailsTab::~SAdvancedPreviewDetailsTab()
 		DefaultSettings->OnAssetViewerSettingsPostUndo().Remove(PostUndoDelegate);
 		DefaultSettings->Save();
 	}
-
-	if (const TSharedPtr<FAdvancedPreviewScene> PreviewScene = PreviewScenePtr.Pin())
-	{
-		PreviewScene->OnProfileChanged().RemoveAll(this);	
-	}
 }
 
 void SAdvancedPreviewDetailsTab::Construct(const FArguments& InArgs, const TSharedRef<FAdvancedPreviewScene>& InPreviewScene)
@@ -142,8 +138,6 @@ void SAdvancedPreviewDetailsTab::Construct(const FArguments& InArgs, const TShar
 
 	UpdateProfileNames();
 	UpdateSettingsView();
-
-	InPreviewScene->OnProfileChanged().AddSP(this, &SAdvancedPreviewDetailsTab::OnProfileChanged);
 }
 
 void SAdvancedPreviewDetailsTab::ComboBoxSelectionChanged(TSharedPtr<FString> NewSelection, ESelectInfo::Type SelectInfo)
@@ -306,14 +300,30 @@ FReply SAdvancedPreviewDetailsTab::CloneProfileButtonClick()
 FReply SAdvancedPreviewDetailsTab::RemoveOrResetProfileButtonClick()
 {
 	const FPreviewSceneProfile& CurrentProfile = DefaultSettings->Profiles[ProfileIndex];
-	
-	if (CurrentProfile.bIsEngineDefaultProfile)
+
+	const FText ActionText = CurrentProfile.bIsEngineDefaultProfile 
+		? LOCTEXT("ResetSceneProfile", "Reset Preview Scene Profile")
+		: LOCTEXT("RemoveSceneProfile", "Remove Preview Scene Profile");
+
+	const EAppReturnType::Type Result = FMessageDialog::Open(
+		EAppMsgType::YesNo,
+		FText::FormatOrdered(
+			LOCTEXT("RemoveOrResetProfileDialogMessage", "Do you want to perform this action?\n{0}"),
+			ActionText
+		)
+	);
+
+	if (Result == EAppReturnType::No || Result == EAppReturnType::Cancel)
 	{
-		const FScopedTransaction Transaction(LOCTEXT("ResetSceneProfile", "Reset Preview Scene Profile"));
-		DefaultSettings->Modify();
+		return FReply::Handled();
+	}
 
-		const FPreviewSceneProfile* DefaultEditorProfile = GetMutableDefault<UDefaultEditorProfiles>()->GetProfile(CurrentProfile.ProfileName);
-		if (DefaultEditorProfile)
+	const FScopedTransaction Transaction(ActionText);
+	DefaultSettings->Modify();
+
+	if (CurrentProfile.bIsEngineDefaultProfile)
+	{
+		if (const FPreviewSceneProfile* DefaultEditorProfile = GetMutableDefault<UDefaultEditorProfiles>()->GetProfile(CurrentProfile.ProfileName))
 		{
 			// Reset currently selected profile 
 			DefaultSettings->Profiles[ProfileIndex] = *DefaultEditorProfile;
@@ -325,15 +335,11 @@ FReply SAdvancedPreviewDetailsTab::RemoveOrResetProfileButtonClick()
 		// in which case we should remove it...
 	}
 
-	const FScopedTransaction Transaction(LOCTEXT("RemoveSceneProfile", "Remove Preview Scene Profile"));
-	DefaultSettings->Modify();
-
 	// Remove currently selected profile 
 	DefaultSettings->Profiles.RemoveAt(ProfileIndex);
 	ProfileIndex = DefaultSettings->Profiles.IsValidIndex(ProfileIndex - 1 ) ? ProfileIndex - 1 : 0;
 	PerProjectSettings->AssetViewerProfileIndex = ProfileIndex;
 	DefaultSettings->PostEditChange();
-	
 	return FReply::Handled();
 }
 
@@ -346,7 +352,9 @@ void SAdvancedPreviewDetailsTab::OnAssetViewerSettingsRefresh(const FName& InPro
 		return;
 	}
 	
-	if (InPropertyName == GET_MEMBER_NAME_CHECKED(FPreviewSceneProfile, ProfileName) || InPropertyName == GET_MEMBER_NAME_CHECKED(FPreviewSceneProfile, bSharedProfile))
+	if (InPropertyName == GET_MEMBER_NAME_CHECKED(FPreviewSceneProfile, ProfileName) 
+		|| InPropertyName == GET_MEMBER_NAME_CHECKED(FPreviewSceneProfile, bSharedProfile)
+		|| ProfileIndex != PreviewScenePtr.Pin()->GetCurrentProfileIndex())
 	{
 		Refresh();
 	}
@@ -399,26 +407,10 @@ void SAdvancedPreviewDetailsTab::OnAssetViewerSettingsPostUndo()
 	PreviewScenePtr.Pin()->UpdateScene(DefaultSettings->Profiles[ProfileIndex]);
 }
 
-void SAdvancedPreviewDetailsTab::OnProfileChanged(const FPreviewSceneProfile& InPreviewSceneProfile, FName InName)
-{
-	if (PerProjectSettings->AssetViewerProfileIndex != ProfileIndex)
-	{
-		Refresh();
-	}
-}
-
 void SAdvancedPreviewDetailsTab::OnPreviewSceneChanged(TSharedRef<FAdvancedPreviewScene> PreviewScene)
 {
-	if (const TSharedPtr<FAdvancedPreviewScene> CurrentPreviewScene = PreviewScenePtr.Pin())
-	{
-		CurrentPreviewScene->OnProfileChanged().RemoveAll(this);
-	}
-
   	PreviewScenePtr = PreviewScene;
 	Refresh();
-
-	PreviewScene->OnProfileChanged().AddSP(this, &SAdvancedPreviewDetailsTab::OnProfileChanged);
-
 }
 
 #undef LOCTEXT_NAMESPACE
\ No newline at end of file
diff --git a/Engine/Source/Editor/AdvancedPreviewScene/Public/SAdvancedPreviewDetailsTab.h b/Engine/Source/Editor/AdvancedPreviewScene/Public/SAdvancedPreviewDetailsTab.h
--- a/Engine/Source/Editor/AdvancedPreviewScene/Public/SAdvancedPreviewDetailsTab.h
+++ b/Engine/Source/Editor/AdvancedPreviewScene/Public/SAdvancedPreviewDetailsTab.h
@@ -58,8 +58,6 @@ protected:
 	UE_API void OnAssetViewerSettingsRefresh(const FName& InPropertyName);
 	UE_API void OnAssetViewerSettingsPostUndo();
 
-	UE_API void OnProfileChanged(const FPreviewSceneProfile& InPreviewSceneProfile, FName InName);
-
 	/** Property viewing widget */
 	TSharedPtr<IDetailsView> SettingsView;
 	TSharedPtr<STextComboBox> ProfileComboBox;
diff --git a/Engine/Source/Editor/UnrealEd/Private/ViewportToolbar/UnrealEdViewportToolbar.cpp b/Engine/Source/Editor/UnrealEd/Private/ViewportToolbar/UnrealEdViewportToolbar.cpp
--- a/Engine/Source/Editor/UnrealEd/Private/ViewportToolbar/UnrealEdViewportToolbar.cpp
+++ b/Engine/Source/Editor/UnrealEd/Private/ViewportToolbar/UnrealEdViewportToolbar.cpp
@@ -3482,7 +3482,24 @@ FToolMenuEntry CreateAssetViewerProfileSubmenu()
 								PreviewProfilesSelectionSection.AddMenuEntry(
 									NAME_None,
 									FText::FromString(ProfileName),
-									FText(),
+									TAttribute<FText>::CreateLambda([]()->FText
+										{
+											const FInputBindingManager& InputManager = FInputBindingManager::Get();
+											const TSharedPtr<FUICommandInfo> PreviousProfile = InputManager.FindCommandInContext(TEXT("AdvancedPreviewScene"), TEXT("PreviousProfile"));
+											const TSharedPtr<FUICommandInfo> NextProfile = InputManager.FindCommandInContext(TEXT("AdvancedPreviewScene"), TEXT("NextProfile"));
+			
+											if (ensureMsgf(PreviousProfile, TEXT("AdvancedPreviewScene PreviousProfile command is not valid"))
+												&& ensureMsgf(NextProfile, TEXT("AdvancedPreviewScene NextProfile command is not valid")))
+											{
+												return FText::FormatOrdered(
+													LOCTEXT("AssetViewerProfileTooltip", "Set this profile as active\nUse {0} to navigate to the previous profile\nUse {1} to navigate to the next profile"),
+													PreviousProfile->GetInputText(),
+													NextProfile->GetInputText()
+												);
+											}
+
+											return FText::GetEmpty();
+										}),
 									FSlateIcon(),
 									FUIAction(
 										FExecuteAction::CreateLambda(
