EditorUsability : PreviewScene
- Added new grid floor material to newly created profile and engine profile
- Added versioning to apply new grid floor material to shared engine profiles as long as they don't have the floor material overriden

#Synced-CL 52419302

diff --git a/Engine/Plugins/Animation/IKRig/Source/IKRigEditor/Private/RetargetEditor/SRetargetAnimAssetsWindow.cpp b/Engine/Plugins/Animation/IKRig/Source/IKRigEditor/Private/RetargetEditor/SRetargetAnimAssetsWindow.cpp
--- a/Engine/Plugins/Animation/IKRig/Source/IKRigEditor/Private/RetargetEditor/SRetargetAnimAssetsWindow.cpp
+++ b/Engine/Plugins/Animation/IKRig/Source/IKRigEditor/Private/RetargetEditor/SRetargetAnimAssetsWindow.cpp
@@ -589,7 +589,7 @@ FRetargetPoseViewportClient::FRetargetPoseViewportClient(FPreviewScene& InPrevie
 	// add a skylight so that models are visible from all angles
 	// TODO, why isn't this working?
 	FPreviewSceneProfile& DefaultProfile = UAssetViewerSettings::Get()->Profiles[GetMutableDefault<UEditorPerProjectUserSettings>()->AssetViewerProfileIndex];
-	DefaultProfile.LoadEnvironmentMap();
+	DefaultProfile.LoadProfileObjects();
 	UTextureCube* CubeMap = DefaultProfile.EnvironmentCubeMap.Get();
 	PreviewScene->SkyLight->SetVisibility(true, false);
 	PreviewScene->SetSkyCubemap(CubeMap);
diff --git a/Engine/Source/Editor/AdvancedPreviewScene/Private/AdvancedPreviewScene.cpp b/Engine/Source/Editor/AdvancedPreviewScene/Private/AdvancedPreviewScene.cpp
--- a/Engine/Source/Editor/AdvancedPreviewScene/Private/AdvancedPreviewScene.cpp
+++ b/Engine/Source/Editor/AdvancedPreviewScene/Private/AdvancedPreviewScene.cpp
@@ -43,7 +43,7 @@ FAdvancedPreviewScene::FAdvancedPreviewScene(ConstructionValues CVS, float InFlo
 	CurrentProfileIndex = DefaultSettings->Profiles.IsValidIndex(SettingsProfileIndex) ? SettingsProfileIndex : 0;
 	ensureMsgf(DefaultSettings->Profiles.IsValidIndex(CurrentProfileIndex), TEXT("Invalid default settings pointer or current profile index"));
 	FPreviewSceneProfile& Profile = DefaultSettings->Profiles[CurrentProfileIndex];
-	Profile.LoadEnvironmentMap();
+	Profile.LoadProfileObjects();
 
 	const FTransform Transform(FRotator(0, 0, 0), FVector(0, 0, 0), FVector(1));
 
@@ -251,7 +251,7 @@ void FAdvancedPreviewScene::SetProfileIndex(const int32 InProfileIndex)
 	}
 
 	CurrentProfileIndex = InProfileIndex;
-	DefaultSettings->Profiles[CurrentProfileIndex].LoadEnvironmentMap();
+	DefaultSettings->Profiles[CurrentProfileIndex].LoadProfileObjects();
 	SetLightDirection(DefaultSettings->Profiles[CurrentProfileIndex].DirectionalLightRotation);
 	PreviousRotation = DefaultSettings->Profiles[CurrentProfileIndex].LightingRigRotation;
 	UILightingRigRotationDelta = 0.0f;
diff --git a/Engine/Source/Editor/AdvancedPreviewScene/Private/AssetViewerSettings.cpp b/Engine/Source/Editor/AdvancedPreviewScene/Private/AssetViewerSettings.cpp
--- a/Engine/Source/Editor/AdvancedPreviewScene/Private/AssetViewerSettings.cpp
+++ b/Engine/Source/Editor/AdvancedPreviewScene/Private/AssetViewerSettings.cpp
@@ -12,6 +12,16 @@ FName UDefaultEditorProfiles::DefaultProfileName = FName("Epic Headquarters");
 FName UDefaultEditorProfiles::EditingProfileName = FName("Grey Wireframe");
 FName UDefaultEditorProfiles::GreyAmbientProfileName = FName("Grey Ambient");
 
+EProfileConfigVersion USharedProfiles::GetVersion() const
+{
+	return Version;
+}
+
+void USharedProfiles::SetVersionToLatest()
+{
+	Version = EProfileConfigVersion::LatestVersion;
+}
+
 UStaticMesh* FPreviewSceneProfile::GetEnvironmentFloorMesh() const
 {
 	if (!EnvironmentFloorMesh && !EnvironmentFloorMeshPath.IsEmpty())
@@ -34,6 +44,13 @@ UMaterialInterface* FPreviewSceneProfile::GetEnvironmentFloorMaterial() const
 	return EnvironmentFloorMaterial.Get();
 }
 
+void FPreviewSceneProfile::LoadProfileObjects()
+{
+	LoadEnvironmentMap();
+	GetEnvironmentFloorMaterial();
+	GetEnvironmentFloorMesh();
+}
+
 void FPreviewSceneProfile::PostEditChangeProperty(FPropertyChangedEvent& InEvent)
 {
 	// Store path to the set environment map texture
@@ -80,7 +97,7 @@ UAssetViewerSettings* UAssetViewerSettings::Get()
 
 		for (FPreviewSceneProfile& Profile : DefaultSettings->Profiles)
 		{
-			Profile.LoadEnvironmentMap();
+			Profile.LoadProfileObjects();
 		}
 		bInitialized = true;
 
@@ -186,11 +203,25 @@ void UAssetViewerSettings::PostInitProperties()
 
 	Profiles.Empty();
 	
-	TArray<FPreviewSceneProfile>& SharedProfiles = GetMutableDefault<USharedProfiles>()->Profiles;
+	USharedProfiles* SharedProfilesDefault = GetMutableDefault<USharedProfiles>();
+	TArray<FPreviewSceneProfile>& SharedProfiles = SharedProfilesDefault->Profiles;
 	for (FPreviewSceneProfile& Profile : SharedProfiles)
 	{
+		// Handle migration here, this is done after config is loaded to update outdated properties
+		if (SharedProfilesDefault->GetVersion() < EProfileConfigVersion::ApplyNewGridMaterial)
+		{
+			// if it's an engine profile and material was not overriden
+			if (Profile.EnvironmentFloorMaterialPath.IsEmpty() && Profile.bIsEngineDefaultProfile)
+			{
+				// Apply new default floor material
+				Profile.EnvironmentFloorMaterialPath = TEXT("/Engine/EngineMaterials/M_Grid.M_Grid");
+			}
+		}
+
 		Profiles.Add(Profile);
 	}
+	// Set the new version to avoid triggering migration again
+	SharedProfilesDefault->SetVersionToLatest();
 
 	TArray<FPreviewSceneProfile>& LocalProfiles = GetMutableDefault<ULocalProfiles>()->Profiles;
 	for (FPreviewSceneProfile& Profile : LocalProfiles)
diff --git a/Engine/Source/Editor/AdvancedPreviewScene/Private/SAdvancedPreviewDetailsTab.cpp b/Engine/Source/Editor/AdvancedPreviewScene/Private/SAdvancedPreviewDetailsTab.cpp
--- a/Engine/Source/Editor/AdvancedPreviewScene/Private/SAdvancedPreviewDetailsTab.cpp
+++ b/Engine/Source/Editor/AdvancedPreviewScene/Private/SAdvancedPreviewDetailsTab.cpp
@@ -209,7 +209,8 @@ FReply SAdvancedPreviewDetailsTab::AddProfileButtonClick()
 	// Add new profile to settings instance
 	ProfileIndex = DefaultSettings->Profiles.AddDefaulted();
 	FPreviewSceneProfile& NewProfile = DefaultSettings->Profiles.Last();
-	
+	NewProfile.LoadProfileObjects();
+
 	// Try to create a valid profile name when one is added
 	bool bFoundValidName = false;
 	int32 ProfileAppendNum = FMath::Max(0, DefaultSettings->Profiles.Num() - 1);
@@ -327,6 +328,7 @@ FReply SAdvancedPreviewDetailsTab::RemoveOrResetProfileButtonClick()
 		{
 			// Reset currently selected profile 
 			DefaultSettings->Profiles[ProfileIndex] = *DefaultEditorProfile;
+			DefaultSettings->Profiles[ProfileIndex].LoadProfileObjects();
 			DefaultSettings->PostEditChange();
 			return FReply::Handled();
 		}
diff --git a/Engine/Source/Editor/AdvancedPreviewScene/Public/AssetViewerSettings.h b/Engine/Source/Editor/AdvancedPreviewScene/Public/AssetViewerSettings.h
--- a/Engine/Source/Editor/AdvancedPreviewScene/Public/AssetViewerSettings.h
+++ b/Engine/Source/Editor/AdvancedPreviewScene/Public/AssetViewerSettings.h
@@ -15,6 +15,16 @@
 
 #define UE_API ADVANCEDPREVIEWSCENE_API
 
+UENUM()
+enum class EProfileConfigVersion : uint8
+{
+	PreVersioning = 0,
+	/** New Material M_Grid applied on shared engine profiles */
+	ApplyNewGridMaterial,
+	VersionPlusOne,
+	LatestVersion = VersionPlusOne - 1
+};
+
 PRAGMA_DISABLE_DEPRECATION_WARNINGS // For bShowGrid
 /**
 * Preview scene profile settings structure.
@@ -44,6 +54,8 @@ struct FPreviewSceneProfile
 		EnvironmentCubeMapPath = TEXT("/Engine/EditorMaterials/AssetViewer/EpicQuadPanorama_CC+EV1.EpicQuadPanorama_CC+EV1");
 		EnvironmentFloorMesh = nullptr;
 		EnvironmentFloorMeshPath = TEXT("/Engine/EditorMeshes/AssetViewer/Floor_Mesh.Floor_Mesh");
+		EnvironmentFloorMaterial = nullptr;
+		EnvironmentFloorMaterialPath = TEXT("/Engine/EngineMaterials/M_Grid.M_Grid");
 		EnvironmentFloorScale = FVector(4.0f, 4.0f, 1.0f);
 		EnvironmentFloorRotation = FRotator();
 		bPostProcessingEnabled = true;
@@ -216,6 +228,9 @@ struct FPreviewSceneProfile
 	/** Loads and returns the environment floor material */
 	UE_API UMaterialInterface* GetEnvironmentFloorMaterial() const;
 	
+	/** Loads required objects for this profile */
+	UE_API void LoadProfileObjects();
+	
 	/** Called when a profile property has been changed */
 	void PostEditChangeProperty(FPropertyChangedEvent& InEvent);
 	
@@ -306,9 +321,17 @@ class USharedProfiles : public UObject
 {
 	GENERATED_BODY()
 public:
+	EProfileConfigVersion GetVersion() const;
+	void SetVersionToLatest();
+
 	/** Collection of shared scene profiles */
 	UPROPERTY(config)
 	TArray<FPreviewSceneProfile> Profiles;
+
+private:
+	/** Used for auto migrations when new changes are introduced */
+	UPROPERTY(Config)
+	EProfileConfigVersion Version = EProfileConfigVersion::PreVersioning;
 };
 
 /**
@@ -320,7 +343,7 @@ class UAssetViewerSettings : public UObject, public FEditorUndoClient
 	GENERATED_BODY()
 public:
 	UE_API UAssetViewerSettings();
-	UE_API virtual ~UAssetViewerSettings();
+	UE_API virtual ~UAssetViewerSettings() override;
 
 	static UE_API UAssetViewerSettings* Get();
 
