Support custom interactions for PreviewScene (sky rotation, light rotation), and disable legacy path appropriately.
Also adds support for static mesh socket manipulation with the new gizmo.

#Synced-CL 52513283

diff --git a/Engine/Source/Editor/AdvancedPreviewScene/AdvancedPreviewScene.Build.cs b/Engine/Source/Editor/AdvancedPreviewScene/AdvancedPreviewScene.Build.cs
--- a/Engine/Source/Editor/AdvancedPreviewScene/AdvancedPreviewScene.Build.cs
+++ b/Engine/Source/Editor/AdvancedPreviewScene/AdvancedPreviewScene.Build.cs
@@ -19,6 +19,7 @@ public class AdvancedPreviewScene : ModuleRules
                 "Core",
                 "CoreUObject",
                 "Engine",
+				"EditorInteractiveToolsFramework",
                 "InputCore",
                 "Slate",
                 "SlateCore",
diff --git a/Engine/Source/Editor/AdvancedPreviewScene/Private/AdvancedPreviewScene.cpp b/Engine/Source/Editor/AdvancedPreviewScene/Private/AdvancedPreviewScene.cpp
--- a/Engine/Source/Editor/AdvancedPreviewScene/Private/AdvancedPreviewScene.cpp
+++ b/Engine/Source/Editor/AdvancedPreviewScene/Private/AdvancedPreviewScene.cpp
@@ -19,6 +19,13 @@
 #include "Components/SphereReflectionCaptureComponent.h"
 
 #include "AssetViewerSettings.h"
+#include "EditorModeManager.h"
+#include "EditorViewportClient.h"
+#include "Tools/EdModeInteractiveToolsContext.h"
+
+#include "ViewportInteractions/ViewportInteractionsBehaviorSource.h"
+#include "ViewportInteractions/ViewportLightRotationInteraction.h"
+#include "ViewportInteractions/ViewportSkyRotationInteraction.h"
 
 #include "Engine/StaticMesh.h"
 #include "AdvancedPreviewSceneCommands.h"
@@ -325,28 +332,34 @@ const bool FAdvancedPreviewScene::HandleViewportInput(FViewport* InViewport, FIn
 {
 	bool bResult = false;
 	const bool bMouseButtonDown = InViewport->KeyState(EKeys::LeftMouseButton) || InViewport->KeyState(EKeys::MiddleMouseButton) || InViewport->KeyState(EKeys::RightMouseButton);
-	
+
 	const bool bSkyMove = InViewport->KeyState(EKeys::K);
 	const bool bLightMoveDown = InViewport->KeyState(EKeys::L);
-	
+
 	// Look at which axis is being dragged and by how much
-	const float DragX = (Key == EKeys::MouseX) ? Delta : 0.f;	
+	const float DragX = (Key == EKeys::MouseX) ? Delta : 0.f;
 	const float DragY = (Key == EKeys::MouseY) ? Delta : 0.f;
 
-	// Move the sky around if K is down and the mouse has moved on the X-axis
-	if (bSkyMove && bMouseButtonDown)
-	{
-		static const float SkyRotationSpeed = 0.22f;
-		float SkyRotation = GetSkyRotation();
-		SkyRotation += -DragX * SkyRotationSpeed;
-		SetSkyRotation(SkyRotation);
-		bResult = true;
-	}
+	// Skip legacy input handling when ITF viewport interactions are enabled — the interactions handle it.
+	bool bUsesViewportInteractions = UsesViewportInteractions(static_cast<FEditorViewportClient*>(InViewport->GetClient()));
 
-	if (bLightMoveDown && (!FMath::IsNearlyZero(DragX) || !FMath::IsNearlyZero(DragY)))
+	if (!bUsesViewportInteractions)
 	{
-		// Save light rotation
-		DefaultSettings->Profiles[CurrentProfileIndex].DirectionalLightRotation = GetLightDirection();
+		// Move the sky around if K is down and the mouse has moved on the X-axis
+		if (bSkyMove && bMouseButtonDown)
+		{
+			static constexpr float SkyRotationSpeed = 0.22f;
+			float SkyRotation = GetSkyRotation();
+			SkyRotation += -DragX * SkyRotationSpeed;
+			SetSkyRotation(SkyRotation);
+			bResult = true;
+		}
+
+		if (bLightMoveDown && (!FMath::IsNearlyZero(DragX) || !FMath::IsNearlyZero(DragY)))
+		{
+			// Save light rotation
+			DefaultSettings->Profiles[CurrentProfileIndex].DirectionalLightRotation = GetLightDirection();
+		}
 	}
 
 	return bResult;
@@ -620,6 +633,27 @@ void FAdvancedPreviewScene::OnAssetViewerSettingsRefresh(const FName& InProperty
 	}
 }
 
+bool FAdvancedPreviewScene::UsesViewportInteractions(const FEditorViewportClient* InViewportClient) const
+{
+	if (!ensure(InViewportClient))
+	{
+		return false;
+	}
+
+	FEditorModeTools* ModeTools = InViewportClient->GetModeTools();
+	if (!ModeTools)
+	{
+		return false;
+	}
+
+	if (const UModeManagerInteractiveToolsContext* ITFContext = ModeTools->GetInteractiveToolsContext())
+	{
+		return ITFContext->UsesViewportInteractions();
+	}
+
+	return false;
+}
+
 void FAdvancedPreviewScene::NavigateProfile(ENavigationDirection InDirection)
 {
 	const bool bMoveForward = InDirection == ENavigationDirection::Forward;
@@ -640,4 +674,13 @@ void FAdvancedPreviewScene::NavigateProfile(ENavigationDirection InDirection)
 	NotificationManager.AddNotification(Info);
 }
 
+void FAdvancedPreviewScene::AddPreviewSceneInteractions(UViewportInteractionsBehaviorSource* InBehaviorSource)
+{
+	if (InBehaviorSource)
+	{
+		InBehaviorSource->AddInteraction<UViewportSkyRotationInteraction>();
+		InBehaviorSource->AddInteraction<UViewportLightRotationInteraction>();
+	}
+}
+
 #undef LOCTEXT_NAMESPACE
diff --git a/Engine/Source/Editor/AdvancedPreviewScene/Private/ViewportInteractions/ViewportLightRotationInteraction.cpp b/Engine/Source/Editor/AdvancedPreviewScene/Private/ViewportInteractions/ViewportLightRotationInteraction.cpp
new file mode 100644
--- /dev/null
+++ b/Engine/Source/Editor/AdvancedPreviewScene/Private/ViewportInteractions/ViewportLightRotationInteraction.cpp
@@ -0,0 +1,128 @@
+// Copyright Epic Games, Inc. All Rights Reserved.
+
+#include "ViewportInteractions/ViewportLightRotationInteraction.h"
+
+#include "AdvancedPreviewScene.h"
+#include "AssetViewerSettings.h"
+#include "Components/DirectionalLightComponent.h"
+#include "EditorViewportClient.h"
+#include "PrimitiveDrawingUtils.h"
+#include "Behaviors/ViewportClickDragBehavior.h"
+#include "SceneView.h"
+#include "Tools/EdModeInteractiveToolsContext.h"
+#include "ViewportInteractions/ViewportInteractionsBehaviorSource.h"
+
+UViewportLightRotationInteraction::UViewportLightRotationInteraction()
+{
+	using namespace UE::Editor::ViewportInteractions;
+
+	InteractionName = TEXT("LightRotation");
+
+	ViewportClickDragBehavior->SetBindings({
+		FButtonBinding(EKeys::LeftMouseButton).TriggersStart().RequiredToStart(false).RequiredToContinue(false),
+		FButtonBinding(EKeys::RightMouseButton).TriggersStart().RequiredToStart(false).RequiredToContinue(false),
+		FButtonBinding(EKeys::MiddleMouseButton).TriggersStart().RequiredToStart(false).RequiredToContinue(false),
+		FButtonBinding(EKeys::L)
+	});
+}
+
+void UViewportLightRotationInteraction::Render(IToolsContextRenderAPI* InRenderAPI)
+{
+	if (!IsActive())
+	{
+		return;
+	}
+
+	const FAdvancedPreviewScene* PreviewScene = GetAdvancedPreviewScene();
+	if (!PreviewScene || !PreviewScene->DirectionalLight)
+	{
+		return;
+	}
+
+	FPrimitiveDrawInterface* PDI = InRenderAPI->GetPrimitiveDrawInterface();
+	const FSceneView* View = InRenderAPI->GetSceneView();
+	if (!PDI || !View)
+	{
+		return;
+	}
+
+	const ULightComponent* Light = PreviewScene->DirectionalLight;
+
+	const FLinearColor ArrowColor = Light->LightColor;
+
+	FTransform LightLocalToWorld = Light->GetComponentToWorld();
+	if (Light->IsA<UDirectionalLightComponent>())
+	{
+		LightLocalToWorld.SetTranslation(FVector::ZeroVector);
+	}
+
+	LightLocalToWorld.SetScale3D(FVector(1.0f));
+
+	// Deproject the current mouse position into world space
+	const FVector2D MouseScreenPos(MouseX, MouseY);
+	FVector MouseWorldPos, MouseWorldDir;
+	View->DeprojectFVector2D(MouseScreenPos, MouseWorldPos, MouseWorldDir);
+
+	// Find the closest point to the origin along the ray
+	MouseWorldPos = FMath::ClosestPointOnLine(MouseWorldPos, MouseWorldPos + MouseWorldDir * WORLD_MAX, FVector::ZeroVector);
+
+	// Compute arrow geometry
+	static constexpr float MinMouseRadius = 100.0f;
+	static constexpr float MinArrowLength = 10.0f;
+	static constexpr float ArrowLengthToSizeRatio = 0.1f;
+	static constexpr float MouseLengthToArrowLengthRatio = 0.2f;
+	static constexpr float ArrowLengthToThicknessRatio = 0.05f;
+	static constexpr float MinArrowThickness = 2.0f;
+
+	const FVector LightToMousePos = MouseWorldPos - LightLocalToWorld.GetTranslation();
+	const float LightToMouseRadius = FMath::Max<FVector::FReal>(LightToMousePos.Size(), MinMouseRadius);
+
+	const float ArrowLength = FMath::Max(MinArrowLength, LightToMouseRadius * MouseLengthToArrowLengthRatio);
+	const float ArrowSize = ArrowLengthToSizeRatio * ArrowLength;
+	const float ArrowThickness = FMath::Max(ArrowLengthToThicknessRatio * ArrowLength, MinArrowThickness);
+
+	const FVector ArrowOrigin = LightLocalToWorld.TransformPosition(FVector(-LightToMouseRadius - 0.5f * ArrowLength, 0.0f, 0.0f));
+	const FQuatRotationTranslationMatrix ArrowToWorld(LightLocalToWorld.GetRotation(), ArrowOrigin);
+
+	DrawDirectionalArrow(PDI, ArrowToWorld, ArrowColor, ArrowLength, ArrowSize, SDPG_World, ArrowThickness);
+}
+
+bool UViewportLightRotationInteraction::CanBeActivated() const
+{
+	if (Super::CanBeActivated())
+	{
+		return GetAdvancedPreviewScene() != nullptr;
+	}
+	return false;
+}
+
+void UViewportLightRotationInteraction::OnDragDelta(float InMouseDeltaX, float InMouseDeltaY)
+{
+	if (FAdvancedPreviewScene* PreviewScene = GetAdvancedPreviewScene())
+	{
+		static constexpr float LightRotSpeed = 0.22f;
+		FRotator LightDir = PreviewScene->GetLightDirection();
+		LightDir.Yaw += -InMouseDeltaX * LightRotSpeed;
+		LightDir.Pitch += -InMouseDeltaY * LightRotSpeed;
+		PreviewScene->SetLightDirection(LightDir);
+
+		// Persist the rotation to the current profile settings
+		if (UAssetViewerSettings* Settings = UAssetViewerSettings::Get())
+		{
+			const int32 ProfileIndex = PreviewScene->GetCurrentProfileIndex();
+			if (Settings->Profiles.IsValidIndex(ProfileIndex))
+			{
+				Settings->Profiles[ProfileIndex].DirectionalLightRotation = LightDir;
+			}
+		}
+	}
+}
+
+FAdvancedPreviewScene* UViewportLightRotationInteraction::GetAdvancedPreviewScene() const
+{
+	if (FEditorViewportClient* ViewportClient = GetEditorViewportClient())
+	{
+		return static_cast<FAdvancedPreviewScene*>(ViewportClient->GetPreviewScene());
+	}
+	return nullptr;
+}
diff --git a/Engine/Source/Editor/AdvancedPreviewScene/Private/ViewportInteractions/ViewportSkyRotationInteraction.cpp b/Engine/Source/Editor/AdvancedPreviewScene/Private/ViewportInteractions/ViewportSkyRotationInteraction.cpp
new file mode 100644
--- /dev/null
+++ b/Engine/Source/Editor/AdvancedPreviewScene/Private/ViewportInteractions/ViewportSkyRotationInteraction.cpp
@@ -0,0 +1,50 @@
+// Copyright Epic Games, Inc. All Rights Reserved.
+
+#include "ViewportInteractions/ViewportSkyRotationInteraction.h"
+
+#include "AdvancedPreviewScene.h"
+#include "EditorViewportClient.h"
+#include "Behaviors/ViewportClickDragBehavior.h"
+
+UViewportSkyRotationInteraction::UViewportSkyRotationInteraction()
+{
+	using namespace UE::Editor::ViewportInteractions;
+
+	InteractionName = TEXT("SkyRotation");
+
+	ViewportClickDragBehavior->SetBindings({
+		FButtonBinding(EKeys::LeftMouseButton).TriggersStart().RequiredToStart(false).RequiredToContinue(false),
+		FButtonBinding(EKeys::RightMouseButton).TriggersStart().RequiredToStart(false).RequiredToContinue(false),
+		FButtonBinding(EKeys::MiddleMouseButton).TriggersStart().RequiredToStart(false).RequiredToContinue(false),
+		FButtonBinding(EKeys::K)
+	});
+}
+
+bool UViewportSkyRotationInteraction::CanBeActivated() const
+{
+	if (Super::CanBeActivated())
+	{
+		return GetAdvancedPreviewScene() != nullptr;
+	}
+	return false;
+}
+
+void UViewportSkyRotationInteraction::OnDragDelta(float InMouseDeltaX, float InMouseDeltaY)
+{
+	if (FAdvancedPreviewScene* PreviewScene = GetAdvancedPreviewScene())
+	{
+		static constexpr float SkyRotationSpeed = 0.22f;
+		float SkyRotation = PreviewScene->GetSkyRotation();
+		SkyRotation += -InMouseDeltaX * SkyRotationSpeed;
+		PreviewScene->SetSkyRotation(SkyRotation);
+	}
+}
+
+FAdvancedPreviewScene* UViewportSkyRotationInteraction::GetAdvancedPreviewScene() const
+{
+	if (FEditorViewportClient* ViewportClient = GetEditorViewportClient())
+	{
+		return static_cast<FAdvancedPreviewScene*>(ViewportClient->GetPreviewScene());
+	}
+	return nullptr;
+}
diff --git a/Engine/Source/Editor/AdvancedPreviewScene/Public/AdvancedPreviewScene.h b/Engine/Source/Editor/AdvancedPreviewScene/Public/AdvancedPreviewScene.h
--- a/Engine/Source/Editor/AdvancedPreviewScene/Public/AdvancedPreviewScene.h
+++ b/Engine/Source/Editor/AdvancedPreviewScene/Public/AdvancedPreviewScene.h
@@ -16,6 +16,7 @@
 
 #define UE_API ADVANCEDPREVIEWSCENE_API
 
+class FEditorViewportClient;
 class FViewport;
 class UAssetViewerSettings;
 class UMaterialInstanceConstant;
@@ -107,6 +108,9 @@ public:
 	/** Sets the floor mesh scale */
 	UE_API void SetFloorMeshScale(const FVector& InScale);
 
+	/** Register sky rotation and light rotation viewport interactions on the given behavior source. */
+	UE_API static void AddPreviewSceneInteractions(class UViewportInteractionsBehaviorSource* InBehaviorSource);
+
 protected:
 	/** Create and map the command list. */
 	UE_API virtual void BindCommands();
@@ -117,6 +121,8 @@ protected:
 	/** Handle refreshing the scene when settings change */
 	UE_API void OnAssetViewerSettingsRefresh(const FName& InPropertyName);
 	
+	UE_API bool UsesViewportInteractions(const FEditorViewportClient* InViewportClient) const;
+
 	/** Navigate to the next or previous profile */
 	enum class ENavigationDirection : uint8
 	{
diff --git a/Engine/Source/Editor/AdvancedPreviewScene/Public/ViewportInteractions/ViewportLightRotationInteraction.h b/Engine/Source/Editor/AdvancedPreviewScene/Public/ViewportInteractions/ViewportLightRotationInteraction.h
new file mode 100644
--- /dev/null
+++ b/Engine/Source/Editor/AdvancedPreviewScene/Public/ViewportInteractions/ViewportLightRotationInteraction.h
@@ -0,0 +1,32 @@
+// Copyright Epic Games, Inc. All Rights Reserved.
+
+#pragma once
+
+#include "ViewportInteractions/ViewportDragInteraction.h"
+#include "ViewportInteractions/ViewportInteractionBindings.h"
+
+#include "ViewportLightRotationInteraction.generated.h"
+
+class FAdvancedPreviewScene;
+class IToolsContextRenderAPI;
+
+/**
+ * Viewport interaction that rotates the directional light.
+ * Used in editors with an FAdvancedPreviewScene (e.g., Static Mesh Editor, Persona).
+ */
+UCLASS(MinimalAPI, Transient)
+class UViewportLightRotationInteraction : public UViewportDragInteraction
+{
+	GENERATED_BODY()
+
+public:
+	ADVANCEDPREVIEWSCENE_API UViewportLightRotationInteraction();
+
+	virtual void Render(IToolsContextRenderAPI* InRenderAPI) override;
+
+	virtual bool CanBeActivated() const override;
+	virtual void OnDragDelta(float InMouseDeltaX, float InMouseDeltaY) override;
+
+private:
+	FAdvancedPreviewScene* GetAdvancedPreviewScene() const;
+};
diff --git a/Engine/Source/Editor/AdvancedPreviewScene/Public/ViewportInteractions/ViewportSkyRotationInteraction.h b/Engine/Source/Editor/AdvancedPreviewScene/Public/ViewportInteractions/ViewportSkyRotationInteraction.h
new file mode 100644
--- /dev/null
+++ b/Engine/Source/Editor/AdvancedPreviewScene/Public/ViewportInteractions/ViewportSkyRotationInteraction.h
@@ -0,0 +1,28 @@
+// Copyright Epic Games, Inc. All Rights Reserved.
+
+#pragma once
+
+#include "ViewportInteractions/ViewportDragInteraction.h"
+
+#include "ViewportSkyRotationInteraction.generated.h"
+
+class FAdvancedPreviewScene;
+
+/**
+ * Viewport interaction that rotates the sky environment.
+ * Used in editors with an FAdvancedPreviewScene (e.g., Static Mesh Editor, Persona).
+ */
+UCLASS(MinimalAPI, Transient)
+class UViewportSkyRotationInteraction : public UViewportDragInteraction
+{
+	GENERATED_BODY()
+
+public:
+	ADVANCEDPREVIEWSCENE_API UViewportSkyRotationInteraction();
+
+	virtual bool CanBeActivated() const override;
+	virtual void OnDragDelta(float InMouseDeltaX, float InMouseDeltaY) override;
+
+private:
+	FAdvancedPreviewScene* GetAdvancedPreviewScene() const;
+};
diff --git a/Engine/Source/Editor/StaticMeshEditor/Private/StaticMeshEditor.cpp b/Engine/Source/Editor/StaticMeshEditor/Private/StaticMeshEditor.cpp
--- a/Engine/Source/Editor/StaticMeshEditor/Private/StaticMeshEditor.cpp
+++ b/Engine/Source/Editor/StaticMeshEditor/Private/StaticMeshEditor.cpp
@@ -1,6 +1,11 @@
 // Copyright Epic Games, Inc. All Rights Reserved.
 
 #include "StaticMeshEditor.h"
+#include "StaticMeshEditorEdMode.h"
+#include "AdvancedPreviewScene.h"
+#include "EditorInteractiveGizmoManager.h"
+#include "EditorGizmos/TransformGizmo.h"
+#include "Elements/Framework/EngineElementsLibrary.h"
 #include "AssetRegistry/AssetData.h"
 #include "Misc/MessageDialog.h"
 #include "HAL/FileManager.h"
@@ -61,8 +66,8 @@
 #include "UnrealExporter.h"
 #include "Exporters/Exporter.h"
 #include "HAL/PlatformApplicationMisc.h"
-#include "AssetEditorModeManager.h"
 #include "StaticMeshEditorViewportClient.h"
+#include "Engine/StaticMeshSocket.h"
 #include "AdvancedPreviewScene.h"
 #include "SnappingUtils.h"
 #include "Tools/EdModeInteractiveToolsContext.h"
@@ -369,6 +374,11 @@ void FStaticMeshEditor::InitStaticMeshEditor( const EToolkitMode::Type Mode, con
 	const bool bCreateDefaultToolbar = true;
 	FAssetEditorToolkit::InitAssetEditor( Mode, InitToolkitHost, StaticMeshEditorAppIdentifier, StandaloneDefaultLayout, bCreateDefaultToolbar, bCreateDefaultStandaloneMenu, ObjectToEdit );
 
+	GetEditorModeManager().SetDefaultMode(FStaticMeshEditorEdMode::Id);
+	GetEditorModeManager().ActivateDefaultMode();
+
+	GetEditorModeManager().SetSupportsViewportITF(true);
+
 	if (UModeManagerInteractiveToolsContext* ToolsContext = GetEditorModeManager().GetInteractiveToolsContext())
 	{
 		ToolsContext->SetViewportInteractionsEnabled(true);
@@ -383,7 +393,6 @@ void FStaticMeshEditor::InitStaticMeshEditor( const EToolkitMode::Type Mode, con
 	check(PinnedToolkitHost.IsValid());
 	ModeUILayer = MakeShareable(new FStaticMeshEditorModeUILayer(PinnedToolkitHost.Get()));
 
-
 	ExtendMenu();
 	ExtendToolBar();
 	RegenerateMenusAndToolbars();
@@ -405,6 +414,12 @@ void FStaticMeshEditor::PostInitAssetEditor()
 	InitSocketManager();
 }
 
+void FStaticMeshEditor::OnBuildViewportInteractions(UViewportInteractionsBehaviorSource* Source)
+{
+	FAssetEditorToolkit::OnBuildViewportInteractions(Source);
+	FAdvancedPreviewScene::AddPreviewSceneInteractions(Source);
+}
+
 void FStaticMeshEditor::GenerateSecondaryToolbar()
 {
 	// Generate the secondary toolbar only if there are registered extensions
diff --git a/Engine/Source/Editor/StaticMeshEditor/Private/StaticMeshEditor.h b/Engine/Source/Editor/StaticMeshEditor/Private/StaticMeshEditor.h
--- a/Engine/Source/Editor/StaticMeshEditor/Private/StaticMeshEditor.h
+++ b/Engine/Source/Editor/StaticMeshEditor/Private/StaticMeshEditor.h
@@ -71,6 +71,7 @@ private:
 	void InitEditorForStaticMesh(UStaticMesh* ObjectToEdit);
 
 	virtual void PostInitAssetEditor() override;
+	virtual void OnBuildViewportInteractions(UViewportInteractionsBehaviorSource* Source) override;
 
 public:
 	virtual void RegisterTabSpawners(const TSharedRef<class FTabManager>& TabManager) override;
diff --git a/Engine/Source/Editor/StaticMeshEditor/Private/StaticMeshEditorEdMode.h b/Engine/Source/Editor/StaticMeshEditor/Private/StaticMeshEditorEdMode.h
new file mode 100644
--- /dev/null
+++ b/Engine/Source/Editor/StaticMeshEditor/Private/StaticMeshEditorEdMode.h
@@ -0,0 +1,17 @@
+// Copyright Epic Games, Inc. All Rights Reserved.
+
+#pragma once
+
+#include "EdMode.h"
+
+/**
+ * EdMode for the StaticMeshEditor, that enables various ITF functionality.
+ */
+class FStaticMeshEditorEdMode : public FEdMode
+{
+public:
+	static inline const FEditorModeID Id = TEXT("StaticMeshEditorEdMode");
+
+	virtual bool ShouldDrawWidget() const override { return true; }
+	virtual bool RequiresLegacyViewportInteractions() const override { return false; }
+};
diff --git a/Engine/Source/Editor/StaticMeshEditor/Private/StaticMeshEditorModule.cpp b/Engine/Source/Editor/StaticMeshEditor/Private/StaticMeshEditorModule.cpp
--- a/Engine/Source/Editor/StaticMeshEditor/Private/StaticMeshEditorModule.cpp
+++ b/Engine/Source/Editor/StaticMeshEditor/Private/StaticMeshEditorModule.cpp
@@ -2,8 +2,10 @@
 
 #include "StaticMeshEditorModule.h"
 #include "Modules/ModuleManager.h"
+#include "EditorModeRegistry.h"
 #include "IStaticMeshEditor.h"
 #include "StaticMeshEditor.h"
+#include "StaticMeshEditorEdMode.h"
 
 const FName StaticMeshEditorAppIdentifier = FName(TEXT("StaticMeshEditorApp"));
 
@@ -26,6 +28,8 @@ public:
 		// Make sure the advanced preview scene module is loaded
 		FModuleManager::Get().LoadModuleChecked("AdvancedPreviewScene");
 
+		FEditorModeRegistry::Get().RegisterMode<FStaticMeshEditorEdMode>(FStaticMeshEditorEdMode::Id);
+
 		MenuExtensibilityManager = MakeShareable(new FExtensibilityManager);
 		ToolBarExtensibilityManager = MakeShareable(new FExtensibilityManager);
 		SecondaryToolBarExtensibilityManager = MakeShareable(new FExtensibilityManager);
@@ -36,6 +40,11 @@ public:
 	 */
 	virtual void ShutdownModule() override
 	{
+		if (UObjectInitialized())
+		{
+			FEditorModeRegistry::Get().UnregisterMode(FStaticMeshEditorEdMode::Id);
+		}
+
 		MenuExtensibilityManager.Reset();
 		ToolBarExtensibilityManager.Reset();
 		SecondaryToolBarExtensibilityManager.Reset();
diff --git a/Engine/Source/Editor/StaticMeshEditor/StaticMeshEditor.Build.cs b/Engine/Source/Editor/StaticMeshEditor/StaticMeshEditor.Build.cs
--- a/Engine/Source/Editor/StaticMeshEditor/StaticMeshEditor.Build.cs
+++ b/Engine/Source/Editor/StaticMeshEditor/StaticMeshEditor.Build.cs
@@ -53,6 +53,8 @@ public class StaticMeshEditor : ModuleRules
 				"StatusBar",
 				"WorkspaceMenuStructure",
 				"PhysicsUtilities",
+				"EditorInteractiveToolsFramework",
+				"InteractiveToolsFramework"
 			}
 		);
 
diff --git a/Engine/Source/Editor/UnrealEd/Private/EditorViewportClient.cpp b/Engine/Source/Editor/UnrealEd/Private/EditorViewportClient.cpp
--- a/Engine/Source/Editor/UnrealEd/Private/EditorViewportClient.cpp
+++ b/Engine/Source/Editor/UnrealEd/Private/EditorViewportClient.cpp
@@ -5346,7 +5346,8 @@ bool FEditorViewportClient::Internal_InputAxis(FViewport* InViewport, FInputDevi
 	const float DragX = (Key == EKeys::MouseX) ? Delta : 0.f;
 	const float DragY = (Key == EKeys::MouseY) ? Delta : 0.f;
 
-	if( bLightMoveDown && bMouseButtonDown && PreviewScene )
+	static FName LightRotationViewportInteractionName("LightRotation");
+	if( bLightMoveDown && bMouseButtonDown && PreviewScene && !IsViewportInteractionActive(ModeTools.Get(), LightRotationViewportInteractionName) )
 	{
 		// Adjust the preview light direction
 		FRotator LightDir = PreviewScene->GetLightDirection();
