EditorUsability : AdvancedPreviewScene
- Updated dialog message and transaction description for preview scene profile reset and removal

#Synced-CL 52573604

diff --git a/Engine/Source/Editor/AdvancedPreviewScene/Private/SAdvancedPreviewDetailsTab.cpp b/Engine/Source/Editor/AdvancedPreviewScene/Private/SAdvancedPreviewDetailsTab.cpp
--- a/Engine/Source/Editor/AdvancedPreviewScene/Private/SAdvancedPreviewDetailsTab.cpp
+++ b/Engine/Source/Editor/AdvancedPreviewScene/Private/SAdvancedPreviewDetailsTab.cpp
@@ -302,16 +302,11 @@ FReply SAdvancedPreviewDetailsTab::RemoveOrResetProfileButtonClick()
 {
 	const FPreviewSceneProfile& CurrentProfile = DefaultSettings->Profiles[ProfileIndex];
 
-	const FText ActionText = CurrentProfile.bIsEngineDefaultProfile 
-		? LOCTEXT("ResetSceneProfile", "Reset Preview Scene Profile")
-		: LOCTEXT("RemoveSceneProfile", "Remove Preview Scene Profile");
-
 	const EAppReturnType::Type Result = FMessageDialog::Open(
 		EAppMsgType::YesNo,
-		FText::FormatOrdered(
-			LOCTEXT("RemoveOrResetProfileDialogMessage", "Do you want to perform this action?\n{0}"),
-			ActionText
-		)
+		CurrentProfile.bIsEngineDefaultProfile 
+			? FText::Format(LOCTEXT("ResetSceneProfileDialogMessage", "Do you want to reset this profile ?\n{0}"), FText::FromString(CurrentProfile.ProfileName))
+			: FText::Format(LOCTEXT("RemoveSceneProfileDialogMessage", "Do you want to remove this profile ?\n{0}"), FText::FromString(CurrentProfile.ProfileName))
 	);
 
 	if (Result == EAppReturnType::No || Result == EAppReturnType::Cancel)
@@ -319,7 +314,11 @@ FReply SAdvancedPreviewDetailsTab::RemoveOrResetProfileButtonClick()
 		return FReply::Handled();
 	}
 
-	const FScopedTransaction Transaction(ActionText);
+	const FScopedTransaction Transaction(
+		CurrentProfile.bIsEngineDefaultProfile 
+			? FText::Format(LOCTEXT("ResetSceneProfile", "Reset Preview Scene Profile ({0})"), FText::FromString(CurrentProfile.ProfileName))
+			: FText::Format(LOCTEXT("RemoveSceneProfile", "Remove Preview Scene Profile ({0})"), FText::FromString(CurrentProfile.ProfileName))
+	);
 	DefaultSettings->Modify();
 
 	if (CurrentProfile.bIsEngineDefaultProfile)
